Fix WMR plugin requesting the native RHI device on the game thread.
This restores the previous method that was used to do this retrieval.

diff --git a/Engine/Plugins/Runtime/WindowsMixedReality/Source/WindowsMixedRealityHMD/Private/WindowsMixedRealityHMD.cpp b/Engine/Plugins/Runtime/WindowsMixedReality/Source/WindowsMixedRealityHMD/Private/WindowsMixedRealityHMD.cpp
--- a/Engine/Plugins/Runtime/WindowsMixedReality/Source/WindowsMixedRealityHMD/Private/WindowsMixedRealityHMD.cpp
+++ b/Engine/Plugins/Runtime/WindowsMixedReality/Source/WindowsMixedRealityHMD/Private/WindowsMixedRealityHMD.cpp
@@ -206,8 +206,18 @@ namespace WindowsMixedReality
 
 	TRefCountPtr<ID3D11Device> FWindowsMixedRealityHMD::InternalGetD3D11Device()
 	{
-		FlushRenderingCommands();
-		return (ID3D11Device*)RHIGetNativeDevice();
+		if (!D3D11Device.IsValid())
+		{
+			FWindowsMixedRealityHMD* Self = this;
+			ENQUEUE_RENDER_COMMAND(InternalGetD3D11DeviceCmd)([Self](FRHICommandListImmediate& RHICmdList)
+			{
+				Self->D3D11Device = (ID3D11Device*)RHIGetNativeDevice();
+			});
+
+			FlushRenderingCommands();
+		}
+
+		return D3D11Device;
 	}
 
 	/** Helper function for acquiring the appropriate FSceneViewport */
