Fix WMR compilation due to missing TRUE/FALSE macros

diff --git a/Engine/Plugins/Runtime/WindowsMixedReality/Source/WindowsMixedRealityHMD/Classes/WindowsMixedRealityStatics.h b/Engine/Plugins/Runtime/WindowsMixedReality/Source/WindowsMixedRealityHMD/Classes/WindowsMixedRealityStatics.h
--- a/Engine/Plugins/Runtime/WindowsMixedReality/Source/WindowsMixedRealityHMD/Classes/WindowsMixedRealityStatics.h
+++ b/Engine/Plugins/Runtime/WindowsMixedReality/Source/WindowsMixedRealityHMD/Classes/WindowsMixedRealityStatics.h
@@ -10,7 +10,9 @@
 #pragma warning(default:4668)
 
 #if WITH_WINDOWS_MIXED_REALITY
+#include "Windows/AllowWindowsPlatformTypes.h"
 #include "MixedRealityInterop.h"
+#include "Windows/HideWindowsPlatformTypes.h"
 #endif
 
 namespace WindowsMixedReality
diff --git a/Engine/Plugins/Runtime/WindowsMixedReality/Source/WindowsMixedRealityHMD/Private/WindowsMixedRealityHMD.h b/Engine/Plugins/Runtime/WindowsMixedReality/Source/WindowsMixedRealityHMD/Private/WindowsMixedRealityHMD.h
--- a/Engine/Plugins/Runtime/WindowsMixedReality/Source/WindowsMixedRealityHMD/Private/WindowsMixedRealityHMD.h
+++ b/Engine/Plugins/Runtime/WindowsMixedReality/Source/WindowsMixedRealityHMD/Private/WindowsMixedRealityHMD.h
@@ -17,7 +17,9 @@
 #include "RendererInterface.h"
 
 #if WITH_WINDOWS_MIXED_REALITY
+#include "Windows/AllowWindowsPlatformTypes.h"
 #include "MixedRealityInterop.h"
+#include "Windows/HideWindowsPlatformTypes.h"
 #endif
 
 namespace WindowsMixedReality
