Ensure we have a valid game viewport before trying to use it.

diff --git a/Engine/Plugins/Runtime/WindowsMixedReality/Source/WindowsMixedRealityHMD/Private/WindowsMixedRealityHMD.cpp b/Engine/Plugins/Runtime/WindowsMixedReality/Source/WindowsMixedRealityHMD/Private/WindowsMixedRealityHMD.cpp
--- a/Engine/Plugins/Runtime/WindowsMixedReality/Source/WindowsMixedRealityHMD/Private/WindowsMixedRealityHMD.cpp
+++ b/Engine/Plugins/Runtime/WindowsMixedReality/Source/WindowsMixedRealityHMD/Private/WindowsMixedRealityHMD.cpp
@@ -402,33 +402,28 @@ namespace WindowsMixedReality
 
 				gameWindowWidth = windowRect.right - windowRect.left;
 				gameWindowHeight = windowRect.bottom - windowRect.top;
-			}
-		}
-
-		// Restore windows focus to game window to preserve keyboard/mouse input.
-		if ((currentWornState == EHMDWornState::Type::Worn) && GEngine)
-		{
-			HWND gameHWND = (HWND)GEngine->GameViewport->GetWindow()->GetNativeWindow()->GetOSWindowHandle();
-
-			// Set mouse focus to center of game window so any clicks interact with the game.
-			if (mouseLockedToCenter)
-			{
-				RECT windowRect;
-				GetWindowRect(gameHWND, &windowRect);
 
-				CenterMouse(windowRect);
-			}
+				// Restore windows focus to game window to preserve keyboard/mouse input.
+				if ((currentWornState == EHMDWornState::Type::Worn) && GEngine)
+				{
+					// Set mouse focus to center of game window so any clicks interact with the game.
+					if (mouseLockedToCenter)
+					{
+						CenterMouse(windowRect);
+					}
 
-			if (GetCapture() != gameHWND)
-			{
-				// Keyboard input
-				SetForegroundWindow(gameHWND);
+					if (GetCapture() != gameHWND)
+					{
+						// Keyboard input
+						SetForegroundWindow(gameHWND);
 
-				// Mouse input
-				SetCapture(gameHWND);
-				SetFocus(gameHWND);
+						// Mouse input
+						SetCapture(gameHWND);
+						SetFocus(gameHWND);
 
-				FSlateApplication::Get().SetAllUserFocusToGameViewport();
+						FSlateApplication::Get().SetAllUserFocusToGameViewport();
+					}
+				}
 			}
 		}
 
