Disable precompiling WindowsMixedRealityHMD for Win32 due to compile errors.

#rnx

diff --git a/Engine/Plugins/Runtime/WindowsMixedReality/Source/WindowsMixedRealityHMD/WindowsMixedRealityHMD.Build.cs b/Engine/Plugins/Runtime/WindowsMixedReality/Source/WindowsMixedRealityHMD/WindowsMixedRealityHMD.Build.cs
--- a/Engine/Plugins/Runtime/WindowsMixedReality/Source/WindowsMixedRealityHMD/WindowsMixedRealityHMD.Build.cs
+++ b/Engine/Plugins/Runtime/WindowsMixedReality/Source/WindowsMixedRealityHMD/WindowsMixedRealityHMD.Build.cs
@@ -129,6 +129,11 @@ namespace UnrealBuildTool.Rules
 				PCHUsage = PCHUsageMode.NoSharedPCHs;
 				PrivatePCHHeaderFile = "Private/WindowsMixedRealityPrecompiled.h";
 			}
+
+			if(Target.Platform == UnrealTargetPlatform.Win32)
+            {
+				PrecompileForTargets = PrecompileTargetsType.None;
+            }
 		}
 	}
 }
