Hacking out WindowsMixedRealityHMD on Win32.

#rnx

diff --git a/Engine/Plugins/Runtime/WindowsMixedReality/Source/WindowsMixedRealityHMD/WindowsMixedRealityHMD.Build.cs b/Engine/Plugins/Runtime/WindowsMixedReality/Source/WindowsMixedRealityHMD/WindowsMixedRealityHMD.Build.cs
--- a/Engine/Plugins/Runtime/WindowsMixedReality/Source/WindowsMixedRealityHMD/WindowsMixedRealityHMD.Build.cs
+++ b/Engine/Plugins/Runtime/WindowsMixedReality/Source/WindowsMixedRealityHMD/WindowsMixedRealityHMD.Build.cs
@@ -42,10 +42,13 @@ namespace UnrealBuildTool.Rules
 					PrivateDefinitions.Add("WINDOWS_MIXED_REALITY_DEBUG_DLL=0");
 				}
 
-				PublicDelayLoadDLLs.Add("HolographicAppRemoting.dll");
-                PublicDelayLoadDLLs.Add("PerceptionDevice.dll");
-				RuntimeDependencies.Add(String.Format("$(EngineDir)/Binaries/ThirdParty/Windows/HoloLens/{0}/HolographicAppRemoting.dll", DllSubpath));
-                RuntimeDependencies.Add(String.Format("$(EngineDir)/Binaries/ThirdParty/Windows/HoloLens/{0}/PerceptionDevice.dll", DllSubpath));
+				if(Target.Platform != UnrealTargetPlatform.Win32)
+                {
+					PublicDelayLoadDLLs.Add("HolographicAppRemoting.dll");
+					PublicDelayLoadDLLs.Add("PerceptionDevice.dll");
+					RuntimeDependencies.Add(String.Format("$(EngineDir)/Binaries/ThirdParty/Windows/HoloLens/{0}/HolographicAppRemoting.dll", DllSubpath));
+					RuntimeDependencies.Add(String.Format("$(EngineDir)/Binaries/ThirdParty/Windows/HoloLens/{0}/PerceptionDevice.dll", DllSubpath));
+                }
             }
 
             PublicDefinitions.Add("WITH_WINDOWS_MIXED_REALITY=1");
@@ -129,11 +132,6 @@ namespace UnrealBuildTool.Rules
 				PCHUsage = PCHUsageMode.NoSharedPCHs;
 				PrivatePCHHeaderFile = "Private/WindowsMixedRealityPrecompiled.h";
 			}
-
-			if(Target.Platform == UnrealTargetPlatform.Win32)
-            {
-				PrecompileForTargets = PrecompileTargetsType.None;
-            }
 		}
 	}
 }
