Fixed up modules that was using "bFasterWithoutUnity" option which is now deprecated to use "bUseUnity" = false instead. But most modules had no reason to disable unity builds and as such I have simply removed the option.

Remember if you are looking to improve your local iteration times you can add the module to BuildConfiguration.xml DisableUnityBuild locally and it will be run without unity for you.

diff --git a/Engine/Plugins/Runtime/AR/Google/GoogleARCore/Source/GoogleARCoreBase/GoogleARCoreBase.Build.cs b/Engine/Plugins/Runtime/AR/Google/GoogleARCore/Source/GoogleARCoreBase/GoogleARCoreBase.Build.cs
--- a/Engine/Plugins/Runtime/AR/Google/GoogleARCore/Source/GoogleARCoreBase/GoogleARCoreBase.Build.cs
+++ b/Engine/Plugins/Runtime/AR/Google/GoogleARCore/Source/GoogleARCoreBase/GoogleARCoreBase.Build.cs
@@ -85,8 +85,6 @@ namespace UnrealBuildTool.Rules
 					RuntimeDependencies.Add("$(EngineDir)/Plugins/Runtime/AR/Google/GoogleARCore/Binaries/ThirdParty/Google/ARCoreImg/" + ExecName);
 				}
 			}
-			
-			bFasterWithoutUnity = false;
         }
     }
 }
diff --git a/Engine/Plugins/Runtime/AR/Google/GoogleARCore/Source/GoogleARCoreRendering/GoogleARCoreRendering.Build.cs b/Engine/Plugins/Runtime/AR/Google/GoogleARCore/Source/GoogleARCoreRendering/GoogleARCoreRendering.Build.cs
--- a/Engine/Plugins/Runtime/AR/Google/GoogleARCore/Source/GoogleARCoreRendering/GoogleARCoreRendering.Build.cs
+++ b/Engine/Plugins/Runtime/AR/Google/GoogleARCore/Source/GoogleARCoreRendering/GoogleARCoreRendering.Build.cs
@@ -28,6 +28,5 @@ public class GoogleARCoreRendering : ModuleRules
 				// ... add private dependencies that you statically link with here ...
 			}
 			);
-		bFasterWithoutUnity = true;
 	}
 }
\ No newline at end of file
diff --git a/Engine/Plugins/Runtime/AR/Google/GoogleARCoreServices/Source/GoogleARCoreServices/GoogleARCoreServices.Build.cs b/Engine/Plugins/Runtime/AR/Google/GoogleARCoreServices/Source/GoogleARCoreServices/GoogleARCoreServices.Build.cs
--- a/Engine/Plugins/Runtime/AR/Google/GoogleARCoreServices/Source/GoogleARCoreServices/GoogleARCoreServices.Build.cs
+++ b/Engine/Plugins/Runtime/AR/Google/GoogleARCoreServices/Source/GoogleARCoreServices/GoogleARCoreServices.Build.cs
@@ -48,8 +48,6 @@ namespace UnrealBuildTool.Rules
 				// Register Plugin Language
 				AdditionalPropertiesForReceipt.Add("AndroidPlugin", Path.Combine(ModuleDirectory, "GoogleARCoreServices_APL.xml"));
 			}
-
-			bFasterWithoutUnity = false;
 		}
 	}
 }
diff --git a/Engine/Plugins/Runtime/Database/SQLiteCore/Source/SQLiteCore/SQLiteCore.Build.cs b/Engine/Plugins/Runtime/Database/SQLiteCore/Source/SQLiteCore/SQLiteCore.Build.cs
--- a/Engine/Plugins/Runtime/Database/SQLiteCore/Source/SQLiteCore/SQLiteCore.Build.cs
+++ b/Engine/Plugins/Runtime/Database/SQLiteCore/Source/SQLiteCore/SQLiteCore.Build.cs
@@ -50,7 +50,7 @@ namespace UnrealBuildTool.Rules
 			}
 
 			bEnableUndefinedIdentifierWarnings = false; // The embedded SQLite implementation generates a lot of these warnings
-			bFasterWithoutUnity = true; // Ensure the embedded SQLite implementation is always compiled in isolation
+			bUseUnity = false; // Ensure the embedded SQLite implementation is always compiled in isolation
 		}
 	}
 }
diff --git a/Engine/Plugins/Runtime/GoogleVR/GoogleVRController/Source/GoogleVRController/GoogleVRController.Build.cs b/Engine/Plugins/Runtime/GoogleVR/GoogleVRController/Source/GoogleVRController/GoogleVRController.Build.cs
--- a/Engine/Plugins/Runtime/GoogleVR/GoogleVRController/Source/GoogleVRController/GoogleVRController.Build.cs
+++ b/Engine/Plugins/Runtime/GoogleVR/GoogleVRController/Source/GoogleVRController/GoogleVRController.Build.cs
@@ -8,8 +8,6 @@ namespace UnrealBuildTool.Rules
 	{
 		public GoogleVRController(ReadOnlyTargetRules Target) : base(Target)
 		{
-			bFasterWithoutUnity = true;
-
 //			string GoogleVRSDKDir = Target.UEThirdPartySourceDirectory + "GoogleVR/";
 			PrivateIncludePaths.AddRange(
 				new string[] {
diff --git a/Engine/Plugins/Runtime/WindowsMixedReality/Source/WindowsMixedRealityEyeTracker/WindowsMixedRealityEyeTracker.Build.cs b/Engine/Plugins/Runtime/WindowsMixedReality/Source/WindowsMixedRealityEyeTracker/WindowsMixedRealityEyeTracker.Build.cs
--- a/Engine/Plugins/Runtime/WindowsMixedReality/Source/WindowsMixedRealityEyeTracker/WindowsMixedRealityEyeTracker.Build.cs
+++ b/Engine/Plugins/Runtime/WindowsMixedReality/Source/WindowsMixedRealityEyeTracker/WindowsMixedRealityEyeTracker.Build.cs
@@ -40,8 +40,6 @@ namespace UnrealBuildTool.Rules
 					}
 				);
 
-				bFasterWithoutUnity = true;
-
                 PublicDefinitions.Add("WITH_WINDOWS_MIXED_REALITY=1");
             }
         }
diff --git a/Engine/Plugins/Runtime/WindowsMixedReality/Source/WindowsMixedRealityHMD/WindowsMixedRealityHMD.Build.cs b/Engine/Plugins/Runtime/WindowsMixedReality/Source/WindowsMixedRealityHMD/WindowsMixedRealityHMD.Build.cs
--- a/Engine/Plugins/Runtime/WindowsMixedReality/Source/WindowsMixedRealityHMD/WindowsMixedRealityHMD.Build.cs
+++ b/Engine/Plugins/Runtime/WindowsMixedReality/Source/WindowsMixedRealityHMD/WindowsMixedRealityHMD.Build.cs
@@ -127,8 +127,6 @@ namespace UnrealBuildTool.Rules
                     PrivateDependencyModuleNames.Add("HoloLensAR");
 				}
 
-                bFasterWithoutUnity = true;
-
 				PCHUsage = PCHUsageMode.NoSharedPCHs;
 				PrivatePCHHeaderFile = "Private/WindowsMixedRealityPrecompiled.h";
 			}
diff --git a/Engine/Plugins/Runtime/WindowsMixedReality/Source/WindowsMixedRealityRHI/WindowsMixedRealityRHI.Build.cs b/Engine/Plugins/Runtime/WindowsMixedReality/Source/WindowsMixedRealityRHI/WindowsMixedRealityRHI.Build.cs
--- a/Engine/Plugins/Runtime/WindowsMixedReality/Source/WindowsMixedRealityRHI/WindowsMixedRealityRHI.Build.cs
+++ b/Engine/Plugins/Runtime/WindowsMixedReality/Source/WindowsMixedRealityRHI/WindowsMixedRealityRHI.Build.cs
@@ -89,8 +89,6 @@ namespace UnrealBuildTool.Rules
                     PrivateIncludePaths.Add("../../../../Source/Runtime/Windows/D3D11RHI/Private/HoloLens");
                 }
 
-                bFasterWithoutUnity = true;
-
                 PCHUsage = PCHUsageMode.NoSharedPCHs;
                 PrivatePCHHeaderFile = "Private/WindowsMixedRealityPrecompiled.h";
             }
diff --git a/Engine/Plugins/Runtime/WindowsMixedReality/Source/WindowsMixedRealitySpatialInput/WindowsMixedRealitySpatialInput.Build.cs b/Engine/Plugins/Runtime/WindowsMixedReality/Source/WindowsMixedRealitySpatialInput/WindowsMixedRealitySpatialInput.Build.cs
--- a/Engine/Plugins/Runtime/WindowsMixedReality/Source/WindowsMixedRealitySpatialInput/WindowsMixedRealitySpatialInput.Build.cs
+++ b/Engine/Plugins/Runtime/WindowsMixedReality/Source/WindowsMixedRealitySpatialInput/WindowsMixedRealitySpatialInput.Build.cs
@@ -63,7 +63,6 @@ namespace UnrealBuildTool.Rules
                 Target.Platform == UnrealTargetPlatform.Win64 ||
                 Target.Platform == UnrealTargetPlatform.HoloLens)
             {	
-                bFasterWithoutUnity = true;
                 bEnableExceptions = true;
 
                 PCHUsage = PCHUsageMode.NoSharedPCHs;
diff --git a/Engine/Source/Runtime/Launch/Launch.Build.cs b/Engine/Source/Runtime/Launch/Launch.Build.cs
--- a/Engine/Source/Runtime/Launch/Launch.Build.cs
+++ b/Engine/Source/Runtime/Launch/Launch.Build.cs
@@ -282,7 +282,7 @@ public class Launch : ModuleRules
 		// @todo ps4 clang bug: this works around a PS4/clang compiler bug (optimizations)
 		if (Target.Platform == UnrealTargetPlatform.PS4)
 		{
-			bFasterWithoutUnity = true;
+			bUseUnity = true;
 		}
 
 		if (Target.IsInPlatformGroup(UnrealPlatformGroup.Unix))
