Remove deprecation warning for MixedReality LiveLink trancking

diff --git a/Engine/Plugins/Runtime/WindowsMixedReality/Source/WindowsMixedRealityHandTracking/Private/LiveLinkWindowsMixedRealityHandTrackingSourceEditor.cpp b/Engine/Plugins/Runtime/WindowsMixedReality/Source/WindowsMixedRealityHandTracking/Private/LiveLinkWindowsMixedRealityHandTrackingSourceEditor.cpp
deleted file mode 100644
diff --git a/Engine/Plugins/Runtime/WindowsMixedReality/Source/WindowsMixedRealityHandTracking/Private/LiveLinkWindowsMixedRealityHandTrackingSourceEditor.h b/Engine/Plugins/Runtime/WindowsMixedReality/Source/WindowsMixedRealityHandTracking/Private/LiveLinkWindowsMixedRealityHandTrackingSourceEditor.h
deleted file mode 100644
diff --git a/Engine/Plugins/Runtime/WindowsMixedReality/Source/WindowsMixedRealityHandTracking/Private/LiveLinkWindowsMixedRealityHandTrackingSourceFactory.cpp b/Engine/Plugins/Runtime/WindowsMixedReality/Source/WindowsMixedRealityHandTracking/Private/LiveLinkWindowsMixedRealityHandTrackingSourceFactory.cpp
--- a/Engine/Plugins/Runtime/WindowsMixedReality/Source/WindowsMixedRealityHandTracking/Private/LiveLinkWindowsMixedRealityHandTrackingSourceFactory.cpp
+++ b/Engine/Plugins/Runtime/WindowsMixedReality/Source/WindowsMixedRealityHandTracking/Private/LiveLinkWindowsMixedRealityHandTrackingSourceFactory.cpp
@@ -1,10 +1,12 @@
 ﻿// Copyright 1998-2019 Epic Games, Inc. All Rights Reserved.
 
 #include "LiveLinkWindowsMixedRealityHandTrackingSourceFactory.h"
-#include "LiveLinkWindowsMixedRealityHandTrackingSourceEditor.h"
 #include "IWindowsMixedRealityHandTrackingPlugin.h"
 #include "WindowsMixedRealityHandTracking.h"
 
+#include "Features/IModularFeatures.h"
+#include "ILiveLinkClient.h"
+
 #define LOCTEXT_NAMESPACE "WindowsMixedRealityHandTracking"
 
 FText ULiveLinkWindowsMixedRealityHandTrackingSourceFactory::GetSourceDisplayName() const
@@ -17,29 +19,23 @@ FText ULiveLinkWindowsMixedRealityHandTrackingSourceFactory::GetSourceTooltip()
 	return LOCTEXT("HandTrackingLiveLinkSourceTooltip", "Windows Mixed Reality Hand Tracking Key Points Source");
 }
 
-TSharedPtr<SWidget> ULiveLinkWindowsMixedRealityHandTrackingSourceFactory::CreateSourceCreationPanel()
+ULiveLinkWindowsMixedRealityHandTrackingSourceFactory::EMenuType ULiveLinkWindowsMixedRealityHandTrackingSourceFactory::GetMenuType() const
 {
-	if (!ActiveSourceEditor.IsValid())
+	if (IModularFeatures::Get().IsModularFeatureAvailable(ILiveLinkClient::ModularFeatureName))
 	{
-		SAssignNew(ActiveSourceEditor, SLiveLinkWindowsMixedRealityHandTrackingSourceEditor);
+		ILiveLinkClient& LiveLinkClient = IModularFeatures::Get().GetModularFeature<ILiveLinkClient>(ILiveLinkClient::ModularFeatureName);
+
+		if (!IWindowsMixedRealityHandTrackingModule::Get().IsLiveLinkSourceValid() || !LiveLinkClient.HasSourceBeenAdded(IWindowsMixedRealityHandTrackingModule::Get().GetLiveLinkSource()))
+		{
+			return EMenuType::MenuEntry;
+		}
 	}
-	return ActiveSourceEditor;
+	return EMenuType::Disabled;
 }
 
-TSharedPtr<ILiveLinkSource> ULiveLinkWindowsMixedRealityHandTrackingSourceFactory::OnSourceCreationPanelClosed(bool bCreateSource)
+TSharedPtr<ILiveLinkSource> ULiveLinkWindowsMixedRealityHandTrackingSourceFactory::CreateSource(const FString& ConnectionString) const
 {
-	TSharedPtr<ILiveLinkSource> NewSource = nullptr;
-
-	if (bCreateSource && ActiveSourceEditor.IsValid())
-	{
-		TSharedPtr<FWindowsMixedRealityHandTracking> HandTracking = StaticCastSharedPtr<FWindowsMixedRealityHandTracking>(IWindowsMixedRealityHandTrackingModule::Get().GetLiveLinkSource());
-
-		// Here we could apply settings from SLiveLinkWindowsMixedRealityHandTrackingSourceEditor
-
-		NewSource = HandTracking;
-	}
-	ActiveSourceEditor = nullptr;
-	return NewSource;
+	return IWindowsMixedRealityHandTrackingModule::Get().GetLiveLinkSource();
 }
 
 #undef LOCTEXT_NAMESPACE
\ No newline at end of file
diff --git a/Engine/Plugins/Runtime/WindowsMixedReality/Source/WindowsMixedRealityHandTracking/Private/LiveLinkWindowsMixedRealityHandTrackingSourceFactory.h b/Engine/Plugins/Runtime/WindowsMixedReality/Source/WindowsMixedRealityHandTracking/Private/LiveLinkWindowsMixedRealityHandTrackingSourceFactory.h
--- a/Engine/Plugins/Runtime/WindowsMixedReality/Source/WindowsMixedRealityHandTracking/Private/LiveLinkWindowsMixedRealityHandTrackingSourceFactory.h
+++ b/Engine/Plugins/Runtime/WindowsMixedReality/Source/WindowsMixedRealityHandTracking/Private/LiveLinkWindowsMixedRealityHandTrackingSourceFactory.h
@@ -15,8 +15,8 @@ public:
 	virtual FText GetSourceDisplayName() const;
 	virtual FText GetSourceTooltip() const;
 
-	virtual TSharedPtr<SWidget> CreateSourceCreationPanel();
-	virtual TSharedPtr<ILiveLinkSource> OnSourceCreationPanelClosed(bool bMakeSource);
+	virtual EMenuType GetMenuType() const override;
+	virtual TSharedPtr<ILiveLinkSource> CreateSource(const FString& ConnectionString) const override;
 
 	TSharedPtr<class SLiveLinkWindowsMixedRealityHandTrackingSourceEditor> ActiveSourceEditor;
 };
\ No newline at end of file
diff --git a/Engine/Plugins/Runtime/WindowsMixedReality/Source/WindowsMixedRealityHandTracking/Private/WindowsMixedRealityHandTracking.cpp b/Engine/Plugins/Runtime/WindowsMixedReality/Source/WindowsMixedRealityHandTracking/Private/WindowsMixedRealityHandTracking.cpp
--- a/Engine/Plugins/Runtime/WindowsMixedReality/Source/WindowsMixedRealityHandTracking/Private/WindowsMixedRealityHandTracking.cpp
+++ b/Engine/Plugins/Runtime/WindowsMixedReality/Source/WindowsMixedRealityHandTracking/Private/WindowsMixedRealityHandTracking.cpp
@@ -81,6 +81,11 @@ public:
 		return InputDevice;
 	}
 
+	virtual bool IsLiveLinkSourceValid() const override
+	{
+		return InputDevice.IsValid();
+	}
+
 	virtual void AddLiveLinkSource() override
 	{
 		if (bLiveLinkSourceRegistered)
diff --git a/Engine/Plugins/Runtime/WindowsMixedReality/Source/WindowsMixedRealityHandTracking/Private/WindowsMixedRealityHandTrackingLiveLink.cpp b/Engine/Plugins/Runtime/WindowsMixedReality/Source/WindowsMixedRealityHandTracking/Private/WindowsMixedRealityHandTrackingLiveLink.cpp
--- a/Engine/Plugins/Runtime/WindowsMixedReality/Source/WindowsMixedRealityHandTracking/Private/WindowsMixedRealityHandTrackingLiveLink.cpp
+++ b/Engine/Plugins/Runtime/WindowsMixedReality/Source/WindowsMixedRealityHandTracking/Private/WindowsMixedRealityHandTrackingLiveLink.cpp
@@ -2,7 +2,6 @@
 
 #include "WindowsMixedRealityHandTracking.h"
 #include "IWindowsMixedRealityHandTrackingPlugin.h"
-#include "LiveLinkWindowsMixedRealityHandTrackingSourceEditor.h"
 #include "CoreMinimal.h"
 #include "UObject/Package.h"
 #include "UObject/UObjectGlobals.h"
diff --git a/Engine/Plugins/Runtime/WindowsMixedReality/Source/WindowsMixedRealityHandTracking/Public/IWindowsMixedRealityHandTrackingPlugin.h b/Engine/Plugins/Runtime/WindowsMixedReality/Source/WindowsMixedRealityHandTracking/Public/IWindowsMixedRealityHandTrackingPlugin.h
--- a/Engine/Plugins/Runtime/WindowsMixedReality/Source/WindowsMixedRealityHandTracking/Public/IWindowsMixedRealityHandTrackingPlugin.h
+++ b/Engine/Plugins/Runtime/WindowsMixedReality/Source/WindowsMixedRealityHandTracking/Public/IWindowsMixedRealityHandTrackingPlugin.h
@@ -49,6 +49,13 @@ public:
 	*/
 	virtual TSharedPtr<class ILiveLinkSource> GetLiveLinkSource() = 0;
 
+	/**
+	 * Checks if the LiveLinkSource as been created.
+	 *
+	 * @return True if the LiveLinkSource has been created with GetLiveLinkSource or AddLiveLinkSource
+	 */
+	virtual bool IsLiveLinkSourceValid() const = 0;
+
 	virtual void AddLiveLinkSource() = 0;
 	virtual void RemoveLiveLinkSource() = 0;
 };
