Fixing WindowsMixedRealityRHI GD3D11RHI merge.
-copying cl 7080409 from dev-vr
-D3D11 was changed a little to support ValidationRHI and WMR must match.

diff --git a/Engine/Plugins/Runtime/WindowsMixedReality/Source/WindowsMixedRealityRHI/Private/WindowsMixedRealityRHI.cpp b/Engine/Plugins/Runtime/WindowsMixedReality/Source/WindowsMixedRealityRHI/Private/WindowsMixedRealityRHI.cpp
--- a/Engine/Plugins/Runtime/WindowsMixedReality/Source/WindowsMixedRealityRHI/Private/WindowsMixedRealityRHI.cpp
+++ b/Engine/Plugins/Runtime/WindowsMixedReality/Source/WindowsMixedRealityRHI/Private/WindowsMixedRealityRHI.cpp
@@ -121,16 +121,19 @@ bool FWindowsMixedRealityRHIModule::IsSupported()
 
 FDynamicRHI* FWindowsMixedRealityRHIModule::CreateRHI(ERHIFeatureLevel::Type RequestedFeatureLevel)
 {
-	FWindowsMixedRealityDynamicRHI * RHI = new FWindowsMixedRealityDynamicRHI(DXGIFactory1.Get(), ChosenAdapter.MaxSupportedFeatureLevel, ChosenAdapter.AdapterIndex, ChosenDescription);
-
 #if PLATFORM_HOLOLENS
 	GMaxRHIFeatureLevel = ERHIFeatureLevel::ES3_1;
 	GMaxRHIShaderPlatform = SP_PCD3D_ES3_1;
+#endif
 
+	FWindowsMixedRealityDynamicRHI * RHI = new FWindowsMixedRealityDynamicRHI(DXGIFactory1.Get(), ChosenAdapter.MaxSupportedFeatureLevel, ChosenAdapter.AdapterIndex, ChosenDescription);
+	GD3D11RHI = RHI;
+
+#if PLATFORM_HOLOLENS
 	RHI->SetHolographicSpace(HoloSpace);
 #endif
 
-	return RHI;
+	return GD3D11RHI;
 }
 
 
