WindowsMixedRealityHandTrackingLiveLink crash fix.
-Hand merging from Dev-vr
-Was crashing due to continued use of a structure which had been MoveTemp-ed to another.
-Now we copy that data before MoveTemp.
-Rearranged the function a bit to get everything into blocks of similar update frequency.

diff --git a/Engine/Plugins/Runtime/WindowsMixedReality/Source/WindowsMixedRealityHandTracking/Private/WindowsMixedRealityHandTracking.h b/Engine/Plugins/Runtime/WindowsMixedReality/Source/WindowsMixedRealityHandTracking/Private/WindowsMixedRealityHandTracking.h
--- a/Engine/Plugins/Runtime/WindowsMixedReality/Source/WindowsMixedRealityHandTracking/Private/WindowsMixedRealityHandTracking.h
+++ b/Engine/Plugins/Runtime/WindowsMixedReality/Source/WindowsMixedRealityHandTracking/Private/WindowsMixedRealityHandTracking.h
@@ -110,6 +110,9 @@ private:
 	FLiveLinkSubjectKey LiveLinkRightHandTrackingSubjectKey;
 	bool bNewLiveLinkClient = false;
 	FLiveLinkStaticDataStruct LiveLinkSkeletonStaticData;
+
+	TArray<FTransform> LeftAnimationTransforms;
+	TArray<FTransform> RightAnimationTransforms;
 };
 
 DEFINE_LOG_CATEGORY_STATIC(LogWindowsMixedRealityHandTracking, Display, All);
diff --git a/Engine/Plugins/Runtime/WindowsMixedReality/Source/WindowsMixedRealityHandTracking/Private/WindowsMixedRealityHandTrackingLiveLink.cpp b/Engine/Plugins/Runtime/WindowsMixedReality/Source/WindowsMixedRealityHandTracking/Private/WindowsMixedRealityHandTrackingLiveLink.cpp
--- a/Engine/Plugins/Runtime/WindowsMixedReality/Source/WindowsMixedRealityHandTracking/Private/WindowsMixedRealityHandTrackingLiveLink.cpp
+++ b/Engine/Plugins/Runtime/WindowsMixedReality/Source/WindowsMixedRealityHandTracking/Private/WindowsMixedRealityHandTrackingLiveLink.cpp
@@ -59,7 +59,11 @@ void FWindowsMixedRealityHandTracking::SetupLiveLinkData()
 {
 	check(IsInGameThread());
 
-	TArray<FName> BoneNames;
+	LiveLinkSkeletonStaticData.InitializeWith(FLiveLinkSkeletonStaticData::StaticStruct(), nullptr);
+	FLiveLinkSkeletonStaticData* SkeletonDataPtr = LiveLinkSkeletonStaticData.Cast<FLiveLinkSkeletonStaticData>();
+	check(SkeletonDataPtr);
+
+	TArray<FName>& BoneNames = SkeletonDataPtr->BoneNames;
 	BoneNames.Reserve(EWMRHandKeypointCount);
 	// Array of bone indices to parent bone index
 	BoneParents.Reserve(EWMRHandKeypointCount);
@@ -107,11 +111,6 @@ void FWindowsMixedRealityHandTracking::SetupLiveLinkData()
 	BoneParents.Add(23);	// LittleDistal -> LittleIntermediate
 	BoneParents.Add(24);	// LittleTip -> LittleDistal
 
-	FLiveLinkStaticDataStruct SkeletonStaticData(FLiveLinkSkeletonStaticData::StaticStruct());
-	LiveLinkSkeletonStaticData.InitializeWith(SkeletonStaticData);
-	FLiveLinkSkeletonStaticData* SkeletonDataPtr = LiveLinkSkeletonStaticData.Cast<FLiveLinkSkeletonStaticData>();
-	check(SkeletonDataPtr);
-	SkeletonDataPtr->SetBoneNames(BoneNames);
 	SkeletonDataPtr->SetBoneParents(BoneParents);
 }
 
@@ -142,33 +141,24 @@ void FWindowsMixedRealityHandTracking::UpdateLiveLink()
 
 	if (LiveLinkClient)
 	{
-		FLiveLinkFrameDataStruct LiveLinkLeftFrameDataStruct(FLiveLinkAnimationFrameData::StaticStruct());
-		FLiveLinkAnimationFrameData& LiveLinkLeftAnimationFrameData = *LiveLinkLeftFrameDataStruct.Cast<FLiveLinkAnimationFrameData>();
-		FLiveLinkFrameDataStruct LiveLinkRightFrameDataStruct(FLiveLinkAnimationFrameData::StaticStruct());
-		FLiveLinkAnimationFrameData& LiveLinkRightAnimationFrameData = *LiveLinkRightFrameDataStruct.Cast<FLiveLinkAnimationFrameData>();
-
-
-		static bool bInitialized = false;
-		if (!bInitialized)
+		// One time initialization:
+		if (LeftAnimationTransforms.Num() == 0)
 		{
+			check(EWMRHandKeypointCount > 0); // ensure the num() test above is a valid way to detect initialization
+
 			SetupLiveLinkData();
 
-			LiveLinkLeftAnimationFrameData.Transforms.Reserve(EWMRHandKeypointCount);
-			LiveLinkRightAnimationFrameData.Transforms.Reserve(EWMRHandKeypointCount);
+			LeftAnimationTransforms.Reserve(EWMRHandKeypointCount);
+			RightAnimationTransforms.Reserve(EWMRHandKeypointCount);
 			// Init to identity all of the Keypoint transforms
 			for (uint32 Count = 0; Count < EWMRHandKeypointCount; ++Count)
 			{
-				LiveLinkLeftAnimationFrameData.Transforms.Add(FTransform::Identity);
-				LiveLinkRightAnimationFrameData.Transforms.Add(FTransform::Identity);
+				LeftAnimationTransforms.Add(FTransform::Identity);
+				RightAnimationTransforms.Add(FTransform::Identity);
 			}
-
-			// Set some metadata to specify handedness
-			LiveLinkLeftAnimationFrameData.MetaData.StringMetaData.Add(FName(TEXT("Handedness")), TEXT("Left"));
-			LiveLinkRightAnimationFrameData.MetaData.StringMetaData.Add(FName(TEXT("Handedness")), TEXT("Right"));
-
-			bInitialized = true;
 		}
 
+		// Per ReceiveClient initialization:
 		if (bNewLiveLinkClient)
 		{
 			FLiveLinkStaticDataStruct SkeletalDataLeft;
@@ -183,15 +173,31 @@ void FWindowsMixedRealityHandTracking::UpdateLiveLink()
 			bNewLiveLinkClient = false;
 		}
 
-		LiveLinkLeftAnimationFrameData.WorldTime = LiveLinkRightAnimationFrameData.WorldTime = FPlatformTime::Seconds();
+		// Every frame updates:
 
 		// Update the transforms for each subject from tracking data
-		UpdateLiveLinkTransforms(LiveLinkLeftAnimationFrameData.Transforms, GetLeftHandState());
-		UpdateLiveLinkTransforms(LiveLinkRightAnimationFrameData.Transforms, GetRightHandState());
+		UpdateLiveLinkTransforms(LeftAnimationTransforms, GetLeftHandState());
+		UpdateLiveLinkTransforms(RightAnimationTransforms, GetRightHandState());
 
-		// Share the data locally with the LiveLink client
-		LiveLinkClient->PushSubjectFrameData_AnyThread(LiveLinkLeftHandTrackingSubjectKey, MoveTemp(LiveLinkLeftFrameDataStruct));
-		LiveLinkClient->PushSubjectFrameData_AnyThread(LiveLinkRightHandTrackingSubjectKey, MoveTemp(LiveLinkRightFrameDataStruct));
+		{
+			// Note these structures will be Moved to live link, leaving them invalid.
+			FLiveLinkFrameDataStruct LiveLinkLeftFrameDataStruct(FLiveLinkAnimationFrameData::StaticStruct());
+			FLiveLinkAnimationFrameData& LiveLinkLeftAnimationFrameData = *LiveLinkLeftFrameDataStruct.Cast<FLiveLinkAnimationFrameData>();
+			FLiveLinkFrameDataStruct LiveLinkRightFrameDataStruct(FLiveLinkAnimationFrameData::StaticStruct());
+			FLiveLinkAnimationFrameData& LiveLinkRightAnimationFrameData = *LiveLinkRightFrameDataStruct.Cast<FLiveLinkAnimationFrameData>();
+			static FName HandednessName (TEXT("Handedness"));
+			LiveLinkLeftAnimationFrameData.MetaData.StringMetaData.Add(HandednessName, TEXT("Left"));
+			LiveLinkRightAnimationFrameData.MetaData.StringMetaData.Add(HandednessName, TEXT("Right"));
+			LiveLinkLeftAnimationFrameData.WorldTime = LiveLinkRightAnimationFrameData.WorldTime = FPlatformTime::Seconds();
+
+			//Copy transforms over to transient structure
+			LiveLinkLeftAnimationFrameData.Transforms = LeftAnimationTransforms;
+			LiveLinkRightAnimationFrameData.Transforms = RightAnimationTransforms;
+
+			// Share the data locally with the LiveLink client
+			LiveLinkClient->PushSubjectFrameData_AnyThread(LiveLinkLeftHandTrackingSubjectKey, MoveTemp(LiveLinkLeftFrameDataStruct));
+			LiveLinkClient->PushSubjectFrameData_AnyThread(LiveLinkRightHandTrackingSubjectKey, MoveTemp(LiveLinkRightFrameDataStruct));
+		}
 	}
 }
 
