Load Hololens libraries by pushing a directory
This avoids a problem due to a circular dependency

diff --git a/Engine/Plugins/Runtime/WindowsMixedReality/Source/WindowsMixedRealityHMD/Private/WindowsMixedRealityHMD.cpp b/Engine/Plugins/Runtime/WindowsMixedReality/Source/WindowsMixedRealityHMD/Private/WindowsMixedRealityHMD.cpp
--- a/Engine/Plugins/Runtime/WindowsMixedReality/Source/WindowsMixedRealityHMD/Private/WindowsMixedRealityHMD.cpp
+++ b/Engine/Plugins/Runtime/WindowsMixedReality/Source/WindowsMixedRealityHMD/Private/WindowsMixedRealityHMD.cpp
@@ -166,7 +166,11 @@ namespace WindowsMixedReality
 
 #if PLATFORM_64BITS
 				// Load these dependencies first or MixedRealityInteropLibraryHandle fails to load since it doesn't look in the correct path for its dependencies automatically
-				FPlatformProcess::GetDllHandle(*(EngineDir / "Binaries/ThirdParty/Windows/x64/HolographicAppRemoting.dll"));
+				FString HoloLensLibraryDir = EngineDir / "Binaries/ThirdParty/Windows/x64";
+				FPlatformProcess::PushDllDirectory(*HoloLensLibraryDir);
+				FPlatformProcess::GetDllHandle(_TEXT("PerceptionDevice.dll"));
+				FPlatformProcess::GetDllHandle(_TEXT("HolographicAppRemoting.dll"));
+				FPlatformProcess::PopDllDirectory(*HoloLensLibraryDir);
 #endif // PLATFORM_64BITS && WITH_EDITOR
 
 				// Then finally try to load the WMR Interop Library
@@ -181,7 +185,7 @@ namespace WindowsMixedReality
 						"Failed to load Windows Mixed Reality Interop Library!  Windows Mixed Reality cannot function.");
 					FMessageDialog::Open(EAppMsgType::Ok, ErrorText);
 					UE_LOG(LogWmrHmd, Error, TEXT("%s"), *ErrorText.ToString());
-				}			
+				}
 			}
 			else
 			{
