- Fixing WindowsMixedRealityHandTracking remap asset for new api

diff --git a/Engine/Plugins/Runtime/WindowsMixedReality/Source/WindowsMixedRealityHandTracking/Private/WindowsMixedRealityHandTrackingLiveLinkRemapAsset.cpp b/Engine/Plugins/Runtime/WindowsMixedReality/Source/WindowsMixedRealityHandTracking/Private/WindowsMixedRealityHandTrackingLiveLinkRemapAsset.cpp
--- a/Engine/Plugins/Runtime/WindowsMixedReality/Source/WindowsMixedRealityHandTracking/Private/WindowsMixedRealityHandTrackingLiveLinkRemapAsset.cpp
+++ b/Engine/Plugins/Runtime/WindowsMixedReality/Source/WindowsMixedRealityHandTracking/Private/WindowsMixedRealityHandTrackingLiveLinkRemapAsset.cpp
@@ -33,7 +33,7 @@ void UWindowsMixedRealityHandTrackingLiveLinkRemapAsset::PostInitProperties()
 }
 #endif
 
-void UWindowsMixedRealityHandTrackingLiveLinkRemapAsset::BuildPoseForSubject(float DeltaTime, const FLiveLinkSkeletonStaticData* InSkeletonData, const FLiveLinkAnimationFrameData* InFrameData, FCompactPose& OutPose, FBlendedCurve& OutCurve)
+void UWindowsMixedRealityHandTrackingLiveLinkRemapAsset::BuildPoseFromAnimationData(float DeltaTime, const FLiveLinkSkeletonStaticData* InSkeletonData, const FLiveLinkAnimationFrameData* InFrameData, FCompactPose& OutPose)
 {
 	check(InSkeletonData);
 	check(InFrameData);
@@ -46,7 +46,7 @@ void UWindowsMixedRealityHandTrackingLiveLinkRemapAsset::BuildPoseForSubject(flo
 
 	for (const FName& SrcBoneName : SourceBoneNames)
 	{
-		TransformedBoneNames.Add( GetRemappedBoneName(SrcBoneName) );
+		TransformedBoneNames.Add(GetRemappedBoneName(SrcBoneName));
 	}
 
 	for (int32 i = 0; i < TransformedBoneNames.Num(); ++i)
@@ -64,7 +64,7 @@ void UWindowsMixedRealityHandTrackingLiveLinkRemapAsset::BuildPoseForSubject(flo
 				FQuat OldRot = BoneTransform.GetRotation();
 				FQuat NewRot;
 				FVector4 OldRotVector(OldRot.X, OldRot.Y, OldRot.Z, OldRot.W);
-				
+
 				NewRot.X = Sign(SwizzleX) * OldRotVector[StaticCast<uint8>(SwizzleX) % 4];
 				NewRot.Y = Sign(SwizzleY) * OldRotVector[StaticCast<uint8>(SwizzleY) % 4];
 				NewRot.Z = Sign(SwizzleZ) * OldRotVector[StaticCast<uint8>(SwizzleZ) % 4];
diff --git a/Engine/Plugins/Runtime/WindowsMixedReality/Source/WindowsMixedRealityHandTracking/Private/WindowsMixedRealityHandTrackingLiveLinkRemapAsset.h b/Engine/Plugins/Runtime/WindowsMixedReality/Source/WindowsMixedRealityHandTracking/Private/WindowsMixedRealityHandTrackingLiveLinkRemapAsset.h
--- a/Engine/Plugins/Runtime/WindowsMixedReality/Source/WindowsMixedRealityHandTracking/Private/WindowsMixedRealityHandTrackingLiveLinkRemapAsset.h
+++ b/Engine/Plugins/Runtime/WindowsMixedReality/Source/WindowsMixedRealityHandTracking/Private/WindowsMixedRealityHandTrackingLiveLinkRemapAsset.h
@@ -34,7 +34,7 @@ class UWindowsMixedRealityHandTrackingLiveLinkRemapAsset :
 	GENERATED_UCLASS_BODY()
 
 public:
-	virtual void BuildPoseForSubject(float DeltaTime, const FLiveLinkSkeletonStaticData* InSkeletonData, const FLiveLinkAnimationFrameData* InFrameData, FCompactPose& OutPose, FBlendedCurve& OutCurve) override;
+	virtual void BuildPoseFromAnimationData(float DeltaTime, const FLiveLinkSkeletonStaticData* InSkeletonData, const FLiveLinkAnimationFrameData* InFrameData, FCompactPose& OutPose) override;
 
 	/** If true, remap the full human hand skeleton including metacarpals */
 	UPROPERTY(EditAnywhere, Category = "LiveLink")
