Merging CL 4428447 from Dev-VR - Update the WMR plugin for Dev-Rendering changes.

diff --git a/Engine/Plugins/Runtime/WindowsMixedReality/Source/WindowsMixedRealityHMD/Private/WindowsMixedRealityHMD.cpp b/Engine/Plugins/Runtime/WindowsMixedReality/Source/WindowsMixedRealityHMD/Private/WindowsMixedRealityHMD.cpp
--- a/Engine/Plugins/Runtime/WindowsMixedReality/Source/WindowsMixedRealityHMD/Private/WindowsMixedRealityHMD.cpp
+++ b/Engine/Plugins/Runtime/WindowsMixedReality/Source/WindowsMixedRealityHMD/Private/WindowsMixedRealityHMD.cpp
@@ -7,6 +7,7 @@
 #include "Modules/ModuleManager.h"
 #include "EngineGlobals.h"
 #include "Engine/Engine.h"
+#include "Interfaces/IPluginManager.h"
 #include "IWindowsMixedRealityHMDPlugin.h"
 
 #if WITH_EDITOR
@@ -92,6 +93,13 @@
 			return FString(TEXT("WindowsMixedRealityHMD"));
 		}
 
+		void StartupModule() override
+		{
+			IHeadMountedDisplayModule::StartupModule();
+			FString PluginShaderDir = FPaths::Combine(IPluginManager::Get().FindPlugin(TEXT("WindowsMixedReality"))->GetBaseDir(), TEXT("Shaders"));
+			AddShaderSourceDirectoryMapping(TEXT("/Plugin/WindowsMixedReality"), PluginShaderDir);
+		}
+
 		void ShutdownModule() override
 		{
 			hmd.Dispose();
@@ -115,7 +123,7 @@
 
 	float FWindowsMixedRealityHMD::GetWorldToMetersScale() const
 	{
-		return GWorld ? GWorld->GetWorldSettings()->WorldToMeters : 100.0f;
+		return CachedWorldToMetersScale;
 	}
 
 	//---------------------------------------------------
@@ -272,6 +280,8 @@
 			}
 		}
 
+		CachedWorldToMetersScale = WorldContext.World()->GetWorldSettings()->WorldToMeters;
+
 		return true;
 	}
 
@@ -376,15 +386,15 @@
 			HeadRotation.Normalize();
 
 			// Position = forward/ backwards, left/ right, up/ down.
-			PositionL = ((FVector(UPoseL.M[2][3], -1 * UPoseL.M[0][3], -1 * UPoseL.M[1][3]) * GWorld->GetWorldSettings()->WorldToMeters));
-			PositionR = ((FVector(UPoseR.M[2][3], -1 * UPoseR.M[0][3], -1 * UPoseR.M[1][3]) * GWorld->GetWorldSettings()->WorldToMeters));
+			PositionL = ((FVector(UPoseL.M[2][3], -1 * UPoseL.M[0][3], -1 * UPoseL.M[1][3]) * GetWorldToMetersScale()));
+			PositionR = ((FVector(UPoseR.M[2][3], -1 * UPoseR.M[0][3], -1 * UPoseR.M[1][3]) * GetWorldToMetersScale()));
 
 			PositionL = RotationL.RotateVector(PositionL);
 			PositionR = RotationR.RotateVector(PositionR);
 
 			if (ipd == 0)
 			{
-				ipd = FVector::Dist(PositionL, PositionR) / GWorld->GetWorldSettings()->WorldToMeters;
+				ipd = FVector::Dist(PositionL, PositionR) / GetWorldToMetersScale();
 			}
 
 			FVector HeadPosition = FMath::Lerp(PositionL, PositionR, 0.5f);
@@ -394,6 +404,49 @@
 		}
 	}
 
+	void SetupHiddenVisibleAreaMesh(TArray<FHMDViewMesh>& HiddenMeshes, TArray<FHMDViewMesh>& VisibleMeshes)
+	{
+		for (int i = (int)MixedRealityInterop::HMDEye::Left;
+			i <= (int)MixedRealityInterop::HMDEye::Right; i++)
+		{
+			MixedRealityInterop::HMDEye eye = (MixedRealityInterop::HMDEye)i;
+
+			DirectX::XMFLOAT2* vertices;
+			int length;
+			if (hmd.GetHiddenAreaMesh(eye, vertices, length))
+			{
+				FVector2D* const vertexPositions = new FVector2D[length];
+				for (int v = 0; v < length; v++)
+				{
+					// Remap to space Unreal is expecting.
+					float x = (vertices[v].x + 1) / 2.0f;
+					float y = (vertices[v].y + 1) / 2.0f;
+
+					vertexPositions[v].Set(x, y);
+				}
+				HiddenMeshes[i].BuildMesh(vertexPositions, length, FHMDViewMesh::MT_HiddenArea);
+
+				delete[] vertexPositions;
+			}
+
+			if (hmd.GetVisibleAreaMesh(eye, vertices, length))
+			{
+				FVector2D* const vertexPositions = new FVector2D[length];
+				for (int v = 0; v < length; v++)
+				{
+					// Remap from NDC space to [0..1] bottom-left origin.
+					float x = (vertices[v].x + 1) / 2.0f;
+					float y = (vertices[v].y + 1) / 2.0f;
+
+					vertexPositions[v].Set(x, y);
+				}
+				VisibleMeshes[i].BuildMesh(vertexPositions, length, FHMDViewMesh::MT_VisibleArea);
+
+				delete[] vertexPositions;
+			}
+		}
+	}
+
 	void FWindowsMixedRealityHMD::SetupHolographicCamera()
 	{
 		// Set the viewport to match the HMD display
@@ -418,7 +471,13 @@
 					bIsStereoEnabled = hmd.IsStereoEnabled();
 					if (bIsStereoEnabled)
 					{
-						SetupHiddenVisibleAreaMesh();
+						hmd.CreateHiddenVisibleAreaMesh();
+
+						FWindowsMixedRealityHMD* Self = this;
+						ENQUEUE_RENDER_COMMAND(SetupHiddenVisibleAreaMeshCmd)([Self](FRHICommandListImmediate& RHICmdList)
+						{
+							SetupHiddenVisibleAreaMesh(Self->HiddenAreaMesh, Self->VisibleAreaMesh);
+						});
 					}
 				}
 				else
@@ -441,51 +500,6 @@
 		GEngine->bForceDisableFrameRateSmoothing = bIsStereoEnabled;
 	}
 
-	void FWindowsMixedRealityHMD::SetupHiddenVisibleAreaMesh()
-	{
-		hmd.CreateHiddenVisibleAreaMesh();
-
-		for (int i = (int)MixedRealityInterop::HMDEye::Left;
-			i <= (int)MixedRealityInterop::HMDEye::Right; i++)
-		{
-			MixedRealityInterop::HMDEye eye = (MixedRealityInterop::HMDEye)i;
-
-			DirectX::XMFLOAT2* vertices;
-			int length;
-			if (hmd.GetHiddenAreaMesh(eye, vertices, length))
-			{
-				FVector2D* const vertexPositions = new FVector2D[length];
-				for (int v = 0; v < length; v++)
-				{
-					// Remap to space Unreal is expecting.
-					float x = (vertices[v].x + 1) / 2.0f;
-					float y = (vertices[v].y + 1) / 2.0f;
-
-					vertexPositions[v].Set(x, y);
-				}
-				HiddenAreaMesh[i].BuildMesh(vertexPositions, length, FHMDViewMesh::MT_HiddenArea);
-
-				delete[] vertexPositions;
-			}
-
-			if (hmd.GetVisibleAreaMesh(eye, vertices, length))
-			{
-				FVector2D* const vertexPositions = new FVector2D[length];
-				for (int v = 0; v < length; v++)
-				{
-					// Remap from NDC space to [0..1] bottom-left origin.
-					float x = (vertices[v].x + 1) / 2.0f;
-					float y = (vertices[v].y + 1) / 2.0f;
-
-					vertexPositions[v].Set(x, y);
-				}
-				VisibleAreaMesh[i].BuildMesh(vertexPositions, length, FHMDViewMesh::MT_VisibleArea);
-
-				delete[] vertexPositions;
-			}
-		}
-	}
-
 	bool FWindowsMixedRealityHMD::IsStereoEnabled() const
 	{
 		return bIsStereoEnabled;
@@ -867,6 +881,9 @@
 	{
 		static const FName RendererModuleName("Renderer");
 		RendererModule = FModuleManager::GetModulePtr<IRendererModule>(RendererModuleName);
+
+		HiddenAreaMesh.SetNum(2);
+		VisibleAreaMesh.SetNum(2);
 	}
 
 	FWindowsMixedRealityHMD::~FWindowsMixedRealityHMD()
diff --git a/Engine/Plugins/Runtime/WindowsMixedReality/Source/WindowsMixedRealityHMD/Private/WindowsMixedRealityHMD.h b/Engine/Plugins/Runtime/WindowsMixedReality/Source/WindowsMixedRealityHMD/Private/WindowsMixedRealityHMD.h
--- a/Engine/Plugins/Runtime/WindowsMixedReality/Source/WindowsMixedRealityHMD/Private/WindowsMixedRealityHMD.h
+++ b/Engine/Plugins/Runtime/WindowsMixedReality/Source/WindowsMixedRealityHMD/Private/WindowsMixedRealityHMD.h
@@ -141,6 +141,7 @@ namespace WindowsMixedReality
 		bool bRequestRestart = false;
 
 		float ScreenScalePercentage = 1.0f;
+		float CachedWorldToMetersScale = 100.0f;
 
 		TRefCountPtr<ID3D11Device> D3D11Device = nullptr;
 
@@ -165,13 +166,12 @@ namespace WindowsMixedReality
 
 		float ipd = 0;
 
-		FHMDViewMesh HiddenAreaMesh[2];
-		FHMDViewMesh VisibleAreaMesh[2];
+		TArray<FHMDViewMesh> HiddenAreaMesh;
+		TArray<FHMDViewMesh> VisibleAreaMesh;
 
 		void StopCustomPresent();
 
 		void SetupHolographicCamera();
-		void SetupHiddenVisibleAreaMesh();
 
 		// Inherited via FXRRenderTargetManager
 		virtual void GetEyeRenderParams_RenderThread(
diff --git a/Engine/Plugins/Runtime/WindowsMixedReality/Source/WindowsMixedRealityHMD/WindowsMixedRealityHMD.Build.cs b/Engine/Plugins/Runtime/WindowsMixedReality/Source/WindowsMixedRealityHMD/WindowsMixedRealityHMD.Build.cs
--- a/Engine/Plugins/Runtime/WindowsMixedReality/Source/WindowsMixedRealityHMD/WindowsMixedRealityHMD.Build.cs
+++ b/Engine/Plugins/Runtime/WindowsMixedReality/Source/WindowsMixedRealityHMD/WindowsMixedRealityHMD.Build.cs
@@ -60,8 +60,9 @@ namespace UnrealBuildTool.Rules
 						"D3D11RHI",
 						"Slate",
 						"SlateCore",
-						"UtilityShaders"
-					}
+						"UtilityShaders",
+						"Projects",
+                    }
 					);
 
 				if (Target.bBuildEditor == true)
