UE4.22 - Fix deprecation
#rnx

diff --git a/Engine/Plugins/Runtime/WindowsMixedReality/Source/WindowsMixedRealityHMD/Private/WindowsMixedRealityHMD.cpp b/Engine/Plugins/Runtime/WindowsMixedReality/Source/WindowsMixedRealityHMD/Private/WindowsMixedRealityHMD.cpp
--- a/Engine/Plugins/Runtime/WindowsMixedReality/Source/WindowsMixedRealityHMD/Private/WindowsMixedRealityHMD.cpp
+++ b/Engine/Plugins/Runtime/WindowsMixedReality/Source/WindowsMixedRealityHMD/Private/WindowsMixedRealityHMD.cpp
@@ -206,18 +206,8 @@ namespace WindowsMixedReality
 
 	TRefCountPtr<ID3D11Device> FWindowsMixedRealityHMD::InternalGetD3D11Device()
 	{
-		TRefCountPtr<ID3D11Device> D3D11DeviceLocal;
-
-		ENQUEUE_UNIQUE_RENDER_COMMAND_ONEPARAMETER(
-			GetNativeDevice,
-			TRefCountPtr<ID3D11Device>&, D3D11DeviceRef, D3D11DeviceLocal,
-			{
-				D3D11DeviceRef = (ID3D11Device*)RHIGetNativeDevice();
-			});
-
 		FlushRenderingCommands();
-
-		return D3D11DeviceLocal;
+		return (ID3D11Device*)RHIGetNativeDevice();
 	}
 
 	/** Helper function for acquiring the appropriate FSceneViewport */
