UE4.22 - Fix deprecation warning
#rnx

diff --git a/Engine/Plugins/Runtime/WindowsMixedReality/Source/WindowsMixedRealityHMD/Private/WindowsMixedRealityHMD.cpp b/Engine/Plugins/Runtime/WindowsMixedReality/Source/WindowsMixedRealityHMD/Private/WindowsMixedRealityHMD.cpp
--- a/Engine/Plugins/Runtime/WindowsMixedReality/Source/WindowsMixedRealityHMD/Private/WindowsMixedRealityHMD.cpp
+++ b/Engine/Plugins/Runtime/WindowsMixedReality/Source/WindowsMixedRealityHMD/Private/WindowsMixedRealityHMD.cpp
@@ -826,8 +826,9 @@ namespace WindowsMixedReality
 
 		const uint32 x = (ViewportWidth - width) * 0.5f;
 		const uint32 y = (ViewportHeight - height) * 0.5f;
-
+		PRAGMA_DISABLE_DEPRECATION_WARNINGS
 		SetRenderTarget(RHICmdList, BackBuffer, FTextureRHIRef());
+		PRAGMA_ENABLE_DEPRECATION_WARNINGS
 		DrawClearQuad(RHICmdList, FLinearColor(0.0f, 0.0f, 0.0f, 1.0f));
 		RHICmdList.SetViewport(x, y, 0, width + x, height + y, 1.0f);
 
