We guard building the WMR plugin to win10 versions that actually support the device. This is not needed when building the release binary on a build machine.

diff --git a/Engine/Plugins/Runtime/WindowsMixedReality/Source/WindowsMixedRealityHMD/WindowsMixedRealityHMD.Build.cs b/Engine/Plugins/Runtime/WindowsMixedReality/Source/WindowsMixedRealityHMD/WindowsMixedRealityHMD.Build.cs
--- a/Engine/Plugins/Runtime/WindowsMixedReality/Source/WindowsMixedRealityHMD/WindowsMixedRealityHMD.Build.cs
+++ b/Engine/Plugins/Runtime/WindowsMixedReality/Source/WindowsMixedRealityHMD/WindowsMixedRealityHMD.Build.cs
@@ -30,7 +30,9 @@ namespace UnrealBuildTool.Rules
             {
                 releaseId = Convert.ToInt32(releaseIdString);
             }
-            bool bAllowWindowsMixedReality = (releaseId >= 1803);
+
+            bool bIsBuildMachine = Environment.GetEnvironmentVariable("IsBuildMachine") == "1";
+            bool bAllowWindowsMixedReality = (releaseId >= 1803) || bIsBuildMachine;
 
             if (bAllowWindowsMixedReality)
             {
