Fix compile errors due to hidden windows platform types.

diff --git a/Engine/Plugins/Runtime/WindowsMixedReality/Source/WindowsMixedRealityHMD/Classes/WindowsMixedRealityFunctionLibrary.h b/Engine/Plugins/Runtime/WindowsMixedReality/Source/WindowsMixedRealityHMD/Classes/WindowsMixedRealityFunctionLibrary.h
--- a/Engine/Plugins/Runtime/WindowsMixedReality/Source/WindowsMixedRealityHMD/Classes/WindowsMixedRealityFunctionLibrary.h
+++ b/Engine/Plugins/Runtime/WindowsMixedReality/Source/WindowsMixedRealityHMD/Classes/WindowsMixedRealityFunctionLibrary.h
@@ -7,7 +7,9 @@
 #include "Kismet/BlueprintFunctionLibrary.h"
 
 #if WITH_WINDOWS_MIXED_REALITY
+#include "Windows/AllowWindowsPlatformTypes.h"
 #include "MixedRealityInterop.h"
+#include "Windows/HideWindowsPlatformTypes.h"
 #endif
 
 #include "WindowsMixedRealityFunctionLibrary.Generated.h"
diff --git a/Engine/Plugins/Runtime/WindowsMixedReality/Source/WindowsMixedRealityHMD/Private/WindowsMixedRealityCustomPresent.h b/Engine/Plugins/Runtime/WindowsMixedReality/Source/WindowsMixedRealityHMD/Private/WindowsMixedRealityCustomPresent.h
--- a/Engine/Plugins/Runtime/WindowsMixedReality/Source/WindowsMixedRealityHMD/Private/WindowsMixedRealityCustomPresent.h
+++ b/Engine/Plugins/Runtime/WindowsMixedReality/Source/WindowsMixedRealityHMD/Private/WindowsMixedRealityCustomPresent.h
@@ -9,7 +9,9 @@
 #include "D3D11Util.h"
 
 #if WITH_WINDOWS_MIXED_REALITY
+#include "Windows/AllowWindowsPlatformTypes.h"
 #include "MixedRealityInterop.h"
+#include "Windows/HideWindowsPlatformTypes.h"
 #endif
 
 namespace WindowsMixedReality
