Increase minimum supported windows 10 version to 1809 for windows mixed reality due to 1803 specific rendering bugs.

diff --git a/Engine/Plugins/Runtime/WindowsMixedReality/Source/WindowsMixedRealityHMD/Private/WindowsMixedRealityHMD.cpp b/Engine/Plugins/Runtime/WindowsMixedReality/Source/WindowsMixedRealityHMD/Private/WindowsMixedRealityHMD.cpp
--- a/Engine/Plugins/Runtime/WindowsMixedReality/Source/WindowsMixedRealityHMD/Private/WindowsMixedRealityHMD.cpp
+++ b/Engine/Plugins/Runtime/WindowsMixedReality/Source/WindowsMixedRealityHMD/Private/WindowsMixedRealityHMD.cpp
@@ -19,8 +19,9 @@
 #include "Windows/WindowsPlatformMisc.h"
 #include "Misc/MessageDialog.h"
 
-// Holographic Remoting is only supported in Windows 10 version 1803 or better
-#define MIN_WIN_10_VERSION_FOR_WMR 1803
+// Holographic Remoting is only supported in Windows 10 version 1809 or better
+// Originally we were supporting 1803, but there were rendering issues specific to that version so for now we only support 1809
+#define MIN_WIN_10_VERSION_FOR_WMR 1809
 
 //---------------------------------------------------
 // Windows Mixed Reality HMD Plugin
@@ -119,17 +120,20 @@ namespace WindowsMixedReality
 
 			void* MixedRealityInteropLibraryHandle = !LibraryPath.IsEmpty() ? FPlatformProcess::GetDllHandle(*LibraryPath) : nullptr;
 
-			if (MixedRealityInteropLibraryHandle)
+			FString OSVersionLabel;
+			FString OSSubVersionLabel;
+			FWindowsPlatformMisc::GetOSVersions(OSVersionLabel, OSSubVersionLabel);
+			// GetOSVersion returns the Win10 release version in the OSVersion rather than the OSSubVersion, so parse it out ourselves
+			OSSubVersionLabel = OSVersionLabel;
+			bool bHasSupportedWindowsVersion = OSSubVersionLabel.RemoveFromStart("Windows 10 (Release ") && OSSubVersionLabel.RemoveFromEnd(")") && (FCString::Atoi(*OSSubVersionLabel) >= MIN_WIN_10_VERSION_FOR_WMR);
+			if (MixedRealityInteropLibraryHandle && bHasSupportedWindowsVersion)
 			{
 				HMD = new MixedRealityInterop();
 			}
 			else
 			{
-				FString OSVersionLabel;
-				FString OSSubVersionLabel;
-				FWindowsPlatformMisc::GetOSVersions(OSVersionLabel, OSSubVersionLabel);
 				FText ErrorText = FText::Format(FTextFormat(NSLOCTEXT("WindowsMixedRealityHMD", "MixedRealityInteropLibraryError", 
-					"Failed to load Windows Mixed Reality Interop Library.\nNote: UE4 only supports Windows Mixed Reality on Windows 10 Release {0} or higher. Current version: {1}")),
+					"Failed to load Windows Mixed Reality Interop Library, or this version of Windows is not supported. \nNote: UE4 only supports Windows Mixed Reality on Windows 10 Release {0} or higher. Current version: {1}")),
 					FText::FromString(FString::FromInt(MIN_WIN_10_VERSION_FOR_WMR)), FText::FromString(OSVersionLabel));
 				FMessageDialog::Open(EAppMsgType::Ok, ErrorText);
 				UE_LOG(LogCore, Error, TEXT("%s"), *ErrorText.ToString());
