Add runtime dependencies to ensure necessary .dlls get copied to staging folder.  Also manually load the dlls MixedRealityInterop.dll depends on from the Engine/Binaries folder before loading it, otherwise it only looks for them in the Project/Binaries folder.

diff --git a/Engine/Plugins/Runtime/WindowsMixedReality/Source/MixedRealityInteropLibrary/MixedRealityInteropLibrary.Build.cs b/Engine/Plugins/Runtime/WindowsMixedReality/Source/MixedRealityInteropLibrary/MixedRealityInteropLibrary.Build.cs
--- a/Engine/Plugins/Runtime/WindowsMixedReality/Source/MixedRealityInteropLibrary/MixedRealityInteropLibrary.Build.cs
+++ b/Engine/Plugins/Runtime/WindowsMixedReality/Source/MixedRealityInteropLibrary/MixedRealityInteropLibrary.Build.cs
@@ -23,5 +23,6 @@ public class MixedRealityInteropLibrary : ModuleRules
 		PublicAdditionalLibraries.Add("MixedRealityInterop.lib");
 		// Delay-load the DLL, so we can load it from the right place first
 		PublicDelayLoadDLLs.Add("MixedRealityInterop.dll");
-	}
+        RuntimeDependencies.Add(PluginDirectory + "/Binaries/ThirdParty/MixedRealityInteropLibrary/" + Target.Platform.ToString() + "/MixedRealityInterop.dll");
+    }
 }
diff --git a/Engine/Plugins/Runtime/WindowsMixedReality/Source/WindowsMixedRealityHMD/Private/WindowsMixedRealityHMD.cpp b/Engine/Plugins/Runtime/WindowsMixedReality/Source/WindowsMixedRealityHMD/Private/WindowsMixedRealityHMD.cpp
--- a/Engine/Plugins/Runtime/WindowsMixedReality/Source/WindowsMixedRealityHMD/Private/WindowsMixedRealityHMD.cpp
+++ b/Engine/Plugins/Runtime/WindowsMixedReality/Source/WindowsMixedRealityHMD/Private/WindowsMixedRealityHMD.cpp
@@ -109,16 +109,21 @@ namespace WindowsMixedReality
 			// Get the base directory of this plugin
 			FString BaseDir = IPluginManager::Get().FindPlugin("WindowsMixedReality")->GetBaseDir();
 
-			// Add on the relative location of the third party dll and load it
-			FString LibraryPath;
+			FString EngineDir = FPaths::EngineDir();
+			FString BinariesSubDir = FPlatformProcess::GetBinariesSubdirectory();
 
-#if PLATFORM_64BITS
-			LibraryPath = FPaths::Combine(*BaseDir, TEXT("Binaries/ThirdParty/MixedRealityInteropLibrary/Win64/MixedRealityInterop.dll"));
-#else 
-			LibraryPath = FPaths::Combine(*BaseDir, TEXT("Binaries/ThirdParty/MixedRealityInteropLibrary/Win32/MixedRealityInterop.dll"));
-#endif // PLATFORM_64BITS
+			FString PerceptionSimulationDLLPath = EngineDir / "Binaries" / BinariesSubDir / "Microsoft.Perception.Simulation.dll";
+			FString HolographicStreamerDesktopDLLPath = EngineDir / "Binaries" / BinariesSubDir / "HolographicStreamerDesktop.dll";
+			FString VCCorLib140DLLPath = EngineDir / "Binaries/ThirdParty/AppLocalDependencies" / BinariesSubDir / "Microsoft.VC140.CRT/vccorlib140.dll";
+			FString MRInteropLibraryPath = BaseDir / "Binaries/ThirdParty/MixedRealityInteropLibrary" / BinariesSubDir / "MixedRealityInterop.dll";
 
-			void* MixedRealityInteropLibraryHandle = !LibraryPath.IsEmpty() ? FPlatformProcess::GetDllHandle(*LibraryPath) : nullptr;
+			// Load these dependencies first or MixedRealityInteropLibraryHandle fails to load since it doesn't look in the correct path for its dependencies automatically
+			void* PerceptionSimulationDLLHandle = FPlatformProcess::GetDllHandle(*PerceptionSimulationDLLPath);
+			void* HolographicStreamerDesktopDLLHandle = FPlatformProcess::GetDllHandle(*HolographicStreamerDesktopDLLPath);
+			void* VCCorLib140DLLHandle = FPlatformProcess::GetDllHandle(*VCCorLib140DLLPath);
+
+			// Then finally try to load the WMR Interop Library
+			void* MixedRealityInteropLibraryHandle = !MRInteropLibraryPath.IsEmpty() ? FPlatformProcess::GetDllHandle(*MRInteropLibraryPath) : nullptr;
 
 			FString OSVersionLabel;
 			FString OSSubVersionLabel;
diff --git a/Engine/Plugins/Runtime/WindowsMixedReality/Source/WindowsMixedRealityHMD/WindowsMixedRealityHMD.Build.cs b/Engine/Plugins/Runtime/WindowsMixedReality/Source/WindowsMixedRealityHMD/WindowsMixedRealityHMD.Build.cs
--- a/Engine/Plugins/Runtime/WindowsMixedReality/Source/WindowsMixedRealityHMD/WindowsMixedRealityHMD.Build.cs
+++ b/Engine/Plugins/Runtime/WindowsMixedReality/Source/WindowsMixedRealityHMD/WindowsMixedRealityHMD.Build.cs
@@ -32,12 +32,14 @@ namespace UnrealBuildTool.Rules
 				RuntimeDependencies.Add("$(EngineDir)/Binaries/Win32/HolographicStreamerDesktop.dll");
 				RuntimeDependencies.Add("$(EngineDir)/Binaries/Win32/Microsoft.Perception.Simulation.dll");
 				RuntimeDependencies.Add("$(EngineDir)/Binaries/Win32/PerceptionSimulationManager.dll");
+				RuntimeDependencies.Add("$(EngineDir)/Binaries/ThirdParty/AppLocalDependencies/Win32/Microsoft.VC140.CRT/vccorlib140.dll");
 			}
 			else if (Target.Platform == UnrealTargetPlatform.Win64)
 			{
 				RuntimeDependencies.Add("$(EngineDir)/Binaries/Win64/HolographicStreamerDesktop.dll");
 				RuntimeDependencies.Add("$(EngineDir)/Binaries/Win64/Microsoft.Perception.Simulation.dll");
 				RuntimeDependencies.Add("$(EngineDir)/Binaries/Win64/PerceptionSimulationManager.dll");
+				RuntimeDependencies.Add("$(EngineDir)/Binaries/ThirdParty/AppLocalDependencies/Win64/Microsoft.VC140.CRT/vccorlib140.dll");
 			}
 
 			PublicLibraryPaths.Add(LibrariesPath);
