Don't force AppLocalDependencies to be included in a packaged build.

diff --git a/Engine/Plugins/Runtime/WindowsMixedReality/Source/WindowsMixedRealityHMD/Private/WindowsMixedRealityHMD.cpp b/Engine/Plugins/Runtime/WindowsMixedReality/Source/WindowsMixedRealityHMD/Private/WindowsMixedRealityHMD.cpp
--- a/Engine/Plugins/Runtime/WindowsMixedReality/Source/WindowsMixedRealityHMD/Private/WindowsMixedRealityHMD.cpp
+++ b/Engine/Plugins/Runtime/WindowsMixedReality/Source/WindowsMixedRealityHMD/Private/WindowsMixedRealityHMD.cpp
@@ -114,13 +114,11 @@ namespace WindowsMixedReality
 
 			FString PerceptionSimulationDLLPath = EngineDir / "Binaries" / BinariesSubDir / "Microsoft.Perception.Simulation.dll";
 			FString HolographicStreamerDesktopDLLPath = EngineDir / "Binaries" / BinariesSubDir / "HolographicStreamerDesktop.dll";
-			FString VCCorLib140DLLPath = EngineDir / "Binaries/ThirdParty/AppLocalDependencies" / BinariesSubDir / "Microsoft.VC140.CRT/vccorlib140.dll";
 			FString MRInteropLibraryPath = BaseDir / "Binaries/ThirdParty/MixedRealityInteropLibrary" / BinariesSubDir / "MixedRealityInterop.dll";
 
 			// Load these dependencies first or MixedRealityInteropLibraryHandle fails to load since it doesn't look in the correct path for its dependencies automatically
 			void* PerceptionSimulationDLLHandle = FPlatformProcess::GetDllHandle(*PerceptionSimulationDLLPath);
 			void* HolographicStreamerDesktopDLLHandle = FPlatformProcess::GetDllHandle(*HolographicStreamerDesktopDLLPath);
-			void* VCCorLib140DLLHandle = FPlatformProcess::GetDllHandle(*VCCorLib140DLLPath);
 
 			// Then finally try to load the WMR Interop Library
 			void* MixedRealityInteropLibraryHandle = !MRInteropLibraryPath.IsEmpty() ? FPlatformProcess::GetDllHandle(*MRInteropLibraryPath) : nullptr;
diff --git a/Engine/Plugins/Runtime/WindowsMixedReality/Source/WindowsMixedRealityHMD/WindowsMixedRealityHMD.Build.cs b/Engine/Plugins/Runtime/WindowsMixedReality/Source/WindowsMixedRealityHMD/WindowsMixedRealityHMD.Build.cs
--- a/Engine/Plugins/Runtime/WindowsMixedReality/Source/WindowsMixedRealityHMD/WindowsMixedRealityHMD.Build.cs
+++ b/Engine/Plugins/Runtime/WindowsMixedReality/Source/WindowsMixedRealityHMD/WindowsMixedRealityHMD.Build.cs
@@ -32,14 +32,12 @@ namespace UnrealBuildTool.Rules
 				RuntimeDependencies.Add("$(EngineDir)/Binaries/Win32/HolographicStreamerDesktop.dll");
 				RuntimeDependencies.Add("$(EngineDir)/Binaries/Win32/Microsoft.Perception.Simulation.dll");
 				RuntimeDependencies.Add("$(EngineDir)/Binaries/Win32/PerceptionSimulationManager.dll");
-				RuntimeDependencies.Add("$(EngineDir)/Binaries/ThirdParty/AppLocalDependencies/Win32/Microsoft.VC140.CRT/vccorlib140.dll");
 			}
 			else if (Target.Platform == UnrealTargetPlatform.Win64)
 			{
 				RuntimeDependencies.Add("$(EngineDir)/Binaries/Win64/HolographicStreamerDesktop.dll");
 				RuntimeDependencies.Add("$(EngineDir)/Binaries/Win64/Microsoft.Perception.Simulation.dll");
 				RuntimeDependencies.Add("$(EngineDir)/Binaries/Win64/PerceptionSimulationManager.dll");
-				RuntimeDependencies.Add("$(EngineDir)/Binaries/ThirdParty/AppLocalDependencies/Win64/Microsoft.VC140.CRT/vccorlib140.dll");
 			}
 
 			PublicLibraryPaths.Add(LibrariesPath);
