LiveLink Windows Mixed Reality Hand Tracking initialization fix
-Was asserting during editor startup because we started using a uobject before it was initialized.

diff --git a/Engine/Plugins/Runtime/WindowsMixedReality/Source/WindowsMixedRealityHandTracking/Private/WindowsMixedRealityHandTracking.cpp b/Engine/Plugins/Runtime/WindowsMixedReality/Source/WindowsMixedRealityHandTracking/Private/WindowsMixedRealityHandTracking.cpp
--- a/Engine/Plugins/Runtime/WindowsMixedReality/Source/WindowsMixedRealityHandTracking/Private/WindowsMixedRealityHandTracking.cpp
+++ b/Engine/Plugins/Runtime/WindowsMixedReality/Source/WindowsMixedRealityHandTracking/Private/WindowsMixedRealityHandTracking.cpp
@@ -125,7 +125,6 @@ FWindowsMixedRealityHandTracking::FWindowsMixedRealityHandTracking(const TShared
 	: MessageHandler(InMessageHandler)
 	, DeviceIndex(0)
 	, bIsHandTrackingStateValid(true)
-	, LiveLinkSkeletonStaticData(FLiveLinkSkeletonStaticData::StaticStruct())
 {
 	// Register "MotionController" modular feature manually
 	IModularFeatures::Get().RegisterModularFeature(GetModularFeatureName(), this);
diff --git a/Engine/Plugins/Runtime/WindowsMixedReality/Source/WindowsMixedRealityHandTracking/Private/WindowsMixedRealityHandTrackingLiveLink.cpp b/Engine/Plugins/Runtime/WindowsMixedReality/Source/WindowsMixedRealityHandTracking/Private/WindowsMixedRealityHandTrackingLiveLink.cpp
--- a/Engine/Plugins/Runtime/WindowsMixedReality/Source/WindowsMixedRealityHandTracking/Private/WindowsMixedRealityHandTrackingLiveLink.cpp
+++ b/Engine/Plugins/Runtime/WindowsMixedReality/Source/WindowsMixedRealityHandTracking/Private/WindowsMixedRealityHandTrackingLiveLink.cpp
@@ -11,6 +11,7 @@
 #include "LiveLinkSourceFactory.h"
 #include "ILiveLinkClient.h"
 #include "Roles/LiveLinkAnimationRole.h"
+#include "Roles/LiveLinkAnimationTypes.h"
 
 #define LOCTEXT_NAMESPACE "WindowsMixedRealityHandTracking"
 
@@ -107,6 +108,8 @@ void FWindowsMixedRealityHandTracking::SetupLiveLinkData()
 	BoneParents.Add(23);	// LittleDistal -> LittleIntermediate
 	BoneParents.Add(24);	// LittleTip -> LittleDistal
 
+	FLiveLinkStaticDataStruct SkeletonStaticData(FLiveLinkSkeletonStaticData::StaticStruct());
+	LiveLinkSkeletonStaticData.InitializeWith(SkeletonStaticData);
 	FLiveLinkSkeletonStaticData* SkeletonDataPtr = LiveLinkSkeletonStaticData.Cast<FLiveLinkSkeletonStaticData>();
 	check(SkeletonDataPtr);
 	SkeletonDataPtr->SetBoneNames(BoneNames);
