Preventing always starting in VR mode when a WMR device is plugged in
Fix for WMR simulator displaying differently in PIE then when using device

diff --git a/Engine/Plugins/Runtime/WindowsMixedReality/Source/WindowsMixedRealityHMD/Private/WindowsMixedRealityHMD.cpp b/Engine/Plugins/Runtime/WindowsMixedReality/Source/WindowsMixedRealityHMD/Private/WindowsMixedRealityHMD.cpp
--- a/Engine/Plugins/Runtime/WindowsMixedReality/Source/WindowsMixedRealityHMD/Private/WindowsMixedRealityHMD.cpp
+++ b/Engine/Plugins/Runtime/WindowsMixedReality/Source/WindowsMixedRealityHMD/Private/WindowsMixedRealityHMD.cpp
@@ -328,8 +328,6 @@ namespace WindowsMixedReality
 
 	void FWindowsMixedRealityHMD::OnBeginPlay(FWorldContext & InWorldContext)
 	{
-		EnableStereo(true);
-
 		//start speech recognition if there are any commands we care to listen for
 		StartSpeechRecognition();
 
@@ -338,8 +336,6 @@ namespace WindowsMixedReality
 
 	void FWindowsMixedRealityHMD::OnEndPlay(FWorldContext & InWorldContext)
 	{
-		EnableStereo(false);
-
 		StopSpeechRecognition();
 
 		IWindowsMixedRealityHandTrackingModule::Get().RemoveLiveLinkSource();
@@ -513,7 +509,7 @@ namespace WindowsMixedReality
 		}
 
 		// Restore windows focus to game window to preserve keyboard/mouse input.
-		if ((currentWornState == EHMDWornState::Type::Worn) && GEngine)
+		if ((currentWornState == EHMDWornState::Type::Worn) && GEngine && GEngine->GameViewport)
 		{
 			HWND gameHWND = (HWND)GEngine->GameViewport->GetWindow()->GetNativeWindow()->GetOSWindowHandle();
 
@@ -1094,7 +1090,9 @@ namespace WindowsMixedReality
 
 	bool FWindowsMixedRealityHMD::HasVisibleAreaMesh() const
 	{
-		return VisibleAreaMesh[0].IsValid() && VisibleAreaMesh[1].IsValid();
+		//re-enable this when we're not running on the simulator once we can query for platform type
+		return false;
+		//return VisibleAreaMesh[0].IsValid() && VisibleAreaMesh[1].IsValid();
 	}
 
 	void FWindowsMixedRealityHMD::DrawVisibleAreaMesh_RenderThread(FRHICommandList& RHICmdList, EStereoscopicPass StereoPass) const
diff --git a/Engine/Source/Editor/UnrealEd/Private/PlayLevel.cpp b/Engine/Source/Editor/UnrealEd/Private/PlayLevel.cpp
--- a/Engine/Source/Editor/UnrealEd/Private/PlayLevel.cpp
+++ b/Engine/Source/Editor/UnrealEd/Private/PlayLevel.cpp
@@ -529,6 +529,13 @@ void UEditorEngine::EndPlayMap()
 		GEngine->PendingDroppedNotes.Empty();
 	}
 
+	//ensure stereo rendering is disabled in case we need to re-enable next PIE run.
+	if (GEngine && GEngine->StereoRenderingDevice)
+	{
+		GEngine->StereoRenderingDevice->EnableStereo(false);
+	}
+
+
 	// Restores realtime viewports that have been disabled for PIE.
 	RestoreRealtimeViewports();
 
