UEVR-1469 Hololens hang on startup
-Reverting to previous stereo enabling behavior on hololens platform.

diff --git a/Engine/Plugins/Runtime/WindowsMixedReality/Source/WindowsMixedRealityHMD/Private/WindowsMixedRealityHMD.cpp b/Engine/Plugins/Runtime/WindowsMixedReality/Source/WindowsMixedRealityHMD/Private/WindowsMixedRealityHMD.cpp
--- a/Engine/Plugins/Runtime/WindowsMixedReality/Source/WindowsMixedRealityHMD/Private/WindowsMixedRealityHMD.cpp
+++ b/Engine/Plugins/Runtime/WindowsMixedReality/Source/WindowsMixedRealityHMD/Private/WindowsMixedRealityHMD.cpp
@@ -328,6 +328,10 @@ namespace WindowsMixedReality
 
 	void FWindowsMixedRealityHMD::OnBeginPlay(FWorldContext & InWorldContext)
 	{
+	#if PLATFORM_HOLOLENS
+		EnableStereo(true);
+	#endif
+
 		//start speech recognition if there are any commands we care to listen for
 		StartSpeechRecognition();
 
@@ -336,6 +340,10 @@ namespace WindowsMixedReality
 
 	void FWindowsMixedRealityHMD::OnEndPlay(FWorldContext & InWorldContext)
 	{
+	#if PLATFORM_HOLOLENS
+		EnableStereo(false);
+	#endif
+
 		StopSpeechRecognition();
 
 		IWindowsMixedRealityHandTrackingModule::Get().RemoveLiveLinkSource();
