UEVR-1448 WMR/Hololens Remoting AR support
-Fixing build warning about dependency on HoloLensAR but it doesn't build on win32
-Made the dependency more specific, though I don't think a win32 windows mixed reality build would actually run.

diff --git a/Engine/Plugins/Runtime/WindowsMixedReality/Source/WindowsMixedRealityHMD/WindowsMixedRealityHMD.Build.cs b/Engine/Plugins/Runtime/WindowsMixedReality/Source/WindowsMixedRealityHMD/WindowsMixedRealityHMD.Build.cs
--- a/Engine/Plugins/Runtime/WindowsMixedReality/Source/WindowsMixedRealityHMD/WindowsMixedRealityHMD.Build.cs
+++ b/Engine/Plugins/Runtime/WindowsMixedReality/Source/WindowsMixedRealityHMD/WindowsMixedRealityHMD.Build.cs
@@ -92,6 +92,11 @@ namespace UnrealBuildTool.Rules
 					}
 					);
 
+				if (Target.Platform == UnrealTargetPlatform.Win64 ||
+					Target.Platform == UnrealTargetPlatform.HoloLens)
+				{
+					PrivateDependencyModuleNames.Add("HoloLensAR");
+				}
 
 				if (Target.bBuildEditor == true)
 				{
