UEVR-1448 WMR/Hololens Remoting AR support
-Previous CL failed to actually remove the bad dependency, just added the narrower one.

diff --git a/Engine/Plugins/Runtime/WindowsMixedReality/Source/WindowsMixedRealityHMD/WindowsMixedRealityHMD.Build.cs b/Engine/Plugins/Runtime/WindowsMixedReality/Source/WindowsMixedRealityHMD/WindowsMixedRealityHMD.Build.cs
--- a/Engine/Plugins/Runtime/WindowsMixedReality/Source/WindowsMixedRealityHMD/WindowsMixedRealityHMD.Build.cs
+++ b/Engine/Plugins/Runtime/WindowsMixedReality/Source/WindowsMixedRealityHMD/WindowsMixedRealityHMD.Build.cs
@@ -88,7 +88,6 @@ namespace UnrealBuildTool.Rules
 						"Projects",
                         "WindowsMixedRealityHandTracking",
 						"AugmentedReality",
-						"HoloLensAR"
 					}
 					);
 
