UEVR-1448 WMR/Hololens Remoting AR support
-Removing win32 from the WindowsMixedReality plugin platform whitelist.  It doesnt' work there due to the lack of AR binaries.

diff --git a/Engine/Plugins/Runtime/WindowsMixedReality/WindowsMixedReality.uplugin b/Engine/Plugins/Runtime/WindowsMixedReality/WindowsMixedReality.uplugin
--- a/Engine/Plugins/Runtime/WindowsMixedReality/WindowsMixedReality.uplugin
+++ b/Engine/Plugins/Runtime/WindowsMixedReality/WindowsMixedReality.uplugin
@@ -15,7 +15,6 @@
 	"IsBetaVersion" : true,
 	"Installed" : false,
 	"SupportedTargetPlatforms": [
-		"Win32",
 		"Win64",
 		"HoloLens"
 	],
@@ -25,7 +24,6 @@
 			"Type": "Runtime",
 			"LoadingPhase": "PostConfigInit",
 			"WhitelistPlatforms": [
-				"Win32",
 				"Win64",
 				"HoloLens"
 			]
@@ -35,7 +33,6 @@
 			"Type": "Runtime",
 			"LoadingPhase": "PostConfigInit",
 			"WhitelistPlatforms": [
-				"Win32",
 				"Win64",
 				"HoloLens"
 			]
@@ -45,7 +42,6 @@
 			"Type": "Runtime",
 			"LoadingPhase": "PreDefault",
 			"WhitelistPlatforms": [
-				"Win32",
 				"Win64",
 				"HoloLens"
 			]
@@ -55,7 +51,6 @@
 			"Type": "Runtime",
 			"LoadingPhase": "PostConfigInit",
 			"WhitelistPlatforms": [
-				"Win32",
 				"Win64",
 				"HoloLens"
 			]
@@ -65,7 +60,6 @@
 			"Type": "Editor",
 			"LoadingPhase": "PostEngineInit",
 			"WhitelistPlatforms": [
-				"Win32",
 				"Win64",
 				"HoloLens"
 			]
@@ -75,7 +69,6 @@
 			"Type": "Runtime",
 			"LoadingPhase": "PreDefault",
 			"WhitelistPlatforms": [
-				"Win32",
 				"Win64",
 				"HoloLens"
 			]
