UE-78077 Editor will not load after enabling WMR plugin - "Failed to preload WindowsMixedRealityHMD.dll"
-Need WMR to always require the HololensAR plugin, not only on hololens because there is a dependency there on both platforms.
-HololensAR plugin will now be automatically enabled when WMR is enabled.

diff --git a/Engine/Plugins/Runtime/WindowsMixedReality/WindowsMixedReality.uplugin b/Engine/Plugins/Runtime/WindowsMixedReality/WindowsMixedReality.uplugin
--- a/Engine/Plugins/Runtime/WindowsMixedReality/WindowsMixedReality.uplugin
+++ b/Engine/Plugins/Runtime/WindowsMixedReality/WindowsMixedReality.uplugin
@@ -95,6 +95,7 @@
 			"Name": "HoloLensAR",
 			"Enabled": true,
 			"WhitelistPlatforms": [
+				"Win64",
 				"HoloLens"
 			]
 		}
