Adding meta data to prompt user to restart editor after turning on "remoting for editor"

diff --git a/Engine/Plugins/Runtime/WindowsMixedReality/Source/WindowsMixedRealityRuntimeSettings/Classes/WindowsMixedRealityRuntimeSettings.h b/Engine/Plugins/Runtime/WindowsMixedReality/Source/WindowsMixedRealityRuntimeSettings/Classes/WindowsMixedRealityRuntimeSettings.h
--- a/Engine/Plugins/Runtime/WindowsMixedReality/Source/WindowsMixedRealityRuntimeSettings/Classes/WindowsMixedRealityRuntimeSettings.h
+++ b/Engine/Plugins/Runtime/WindowsMixedReality/Source/WindowsMixedRealityRuntimeSettings/Classes/WindowsMixedRealityRuntimeSettings.h
@@ -31,7 +31,7 @@ public:
 	virtual void PostEditChangeProperty(struct FPropertyChangedEvent& PropertyChangedEvent) override;
 #endif // WITH_EDITOR
 
-	UPROPERTY(GlobalConfig, EditAnywhere, Category = "Holographic Remoting", Meta = (DisplayName = "Enable Remoting For Editor (Requires Restart)", Tooltip = "If true WMR is a valid HMD even if none is connected so that one could connect via remoting.  Editor restart required."))
+	UPROPERTY(GlobalConfig, EditAnywhere, Category = "Holographic Remoting", Meta = (ConfigRestartRequired = true, DisplayName = "Enable Remoting For Editor (Requires Restart)", Tooltip = "If true WMR is a valid HMD even if none is connected so that one could connect via remoting.  Editor restart required."))
 	bool bEnableRemotingForEditor = false;
 
 	/** The IP of the HoloLens to remote to. */
