Update GetRelativeEyePose to actually get relative eye poses

Use base implementation for CalculateStereoViewOffset

diff --git a/Engine/Plugins/Runtime/WindowsMixedReality/Source/WindowsMixedRealityHMD/Private/WindowsMixedRealityHMD.cpp b/Engine/Plugins/Runtime/WindowsMixedReality/Source/WindowsMixedRealityHMD/Private/WindowsMixedRealityHMD.cpp
--- a/Engine/Plugins/Runtime/WindowsMixedReality/Source/WindowsMixedRealityHMD/Private/WindowsMixedRealityHMD.cpp
+++ b/Engine/Plugins/Runtime/WindowsMixedReality/Source/WindowsMixedRealityHMD/Private/WindowsMixedRealityHMD.cpp
@@ -620,15 +620,16 @@ namespace WindowsMixedReality
 	{
 		OutOrientation = FQuat::Identity;
 		OutPosition = FVector::ZeroVector;
-		if (DeviceId == IXRTrackingSystem::HMDDeviceId && (Eye == eSSP_LEFT_EYE || Eye == eSSP_RIGHT_EYE))
-		{
-			OutPosition = FVector(0, (Eye == EStereoscopicPass::eSSP_LEFT_EYE ? -0.5f : 0.5f) * GetInterpupillaryDistance() * GetWorldToMetersScale(), 0);
-			return true;
-		}
-		else
+		if (Eye != eSSP_LEFT_EYE && Eye != eSSP_RIGHT_EYE)
 		{
 			return false;
 		}
+
+		const FTransform relativeTransform = Eye == eSSP_LEFT_EYE ? LeftTransform * HeadTransform.Inverse() : RightTransform * HeadTransform.Inverse();
+		OutPosition = relativeTransform.GetTranslation();
+		OutOrientation = relativeTransform.GetRotation();
+
+		return true;
 	}
 
 	void FWindowsMixedRealityHMD::ResetOrientationAndPosition(float yaw)
@@ -681,6 +682,10 @@ namespace WindowsMixedReality
 
 			CurrOrientation = HeadRotation;
 			CurrPosition = HeadPosition;
+
+			LeftTransform = FTransform(RotationL, PositionL, FVector::OneVector);
+			RightTransform = FTransform(RotationR, PositionR, FVector::OneVector);
+			HeadTransform = FTransform(CurrOrientation, CurrPosition, FVector::OneVector);
 		}
 #endif
 	}
@@ -845,34 +850,6 @@ namespace WindowsMixedReality
 		}
 	}
 
-	void FWindowsMixedRealityHMD::CalculateStereoViewOffset(const enum EStereoscopicPass StereoPassType, FRotator& ViewRotation, const float WorldToMeters, FVector& ViewLocation)
-	{
-		if (StereoPassType != eSSP_LEFT_EYE &&
-			StereoPassType != eSSP_RIGHT_EYE)
-		{
-			return;
-		}
-
-		FVector HmdToEyeOffset = FVector::ZeroVector;
-
-		FQuat CurEyeOrient;
-		if (StereoPassType == eSSP_LEFT_EYE)
-		{
-			HmdToEyeOffset = PositionL - CurrPosition;
-			CurEyeOrient = RotationL;
-		}
-		else if (StereoPassType == eSSP_RIGHT_EYE)
-		{
-			HmdToEyeOffset = PositionR - CurrPosition;
-			CurEyeOrient = RotationR;
-		}
-
-		const FQuat ViewOrient = ViewRotation.Quaternion();
-		const FQuat deltaControlOrientation = ViewOrient * CurEyeOrient.Inverse();
-		const FVector vEyePosition = deltaControlOrientation.RotateVector(HmdToEyeOffset);
-		ViewLocation += vEyePosition;
-	}
-
 	FMatrix FWindowsMixedRealityHMD::GetStereoProjectionMatrix(const enum EStereoscopicPass StereoPassType) const
 	{
 		if (StereoPassType != eSSP_LEFT_EYE &&
diff --git a/Engine/Plugins/Runtime/WindowsMixedReality/Source/WindowsMixedRealityHMD/Private/WindowsMixedRealityHMD.h b/Engine/Plugins/Runtime/WindowsMixedReality/Source/WindowsMixedRealityHMD/Private/WindowsMixedRealityHMD.h
--- a/Engine/Plugins/Runtime/WindowsMixedReality/Source/WindowsMixedRealityHMD/Private/WindowsMixedRealityHMD.h
+++ b/Engine/Plugins/Runtime/WindowsMixedReality/Source/WindowsMixedRealityHMD/Private/WindowsMixedRealityHMD.h
@@ -107,8 +107,6 @@ namespace WindowsMixedReality
 			EStereoscopicPass StereoPass,
 			int32& X, int32& Y,
 			uint32& SizeX, uint32& SizeY) const override;
-		virtual void CalculateStereoViewOffset(const EStereoscopicPass StereoPassType, FRotator& ViewRotation,
-			const float MetersToWorld, FVector& ViewLocation) override;
 		virtual FMatrix GetStereoProjectionMatrix(const enum EStereoscopicPass StereoPassType) const override;
 		virtual IStereoRenderTargetManager* GetRenderTargetManager() override { return this; }
 		virtual class IStereoLayers* GetStereoLayers() override;
@@ -185,6 +183,9 @@ namespace WindowsMixedReality
 		FQuat RotationR = FQuat::Identity;
 		FVector PositionL = FVector::ZeroVector;
 		FVector PositionR = FVector::ZeroVector;
+		FTransform LeftTransform = FTransform::Identity;
+		FTransform RightTransform = FTransform::Identity;
+		FTransform HeadTransform = FTransform::Identity;
 
 		float ipd = 0;
 
