UE-79183 Holographic Frames presented out of sequence when running BP_ComposurePostMoveCroppingMatrixTest_00 tests in ComposurePostMoves map in EngineTest project
-Fixed asserts and crashes during screenshot on hololens which running this test revealed.  I did not see the bug as described..
-Screenshot taking appears to uncover some bad assumptions in the hololens rendering code, and an unused back buffer ref.
-After this fix the test eventually ran out of memory, again no frame present sequence error was seen.

diff --git a/Engine/Plugins/Runtime/WindowsMixedReality/Source/WindowsMixedRealityHMD/Private/WindowsMixedRealityHMD.cpp b/Engine/Plugins/Runtime/WindowsMixedReality/Source/WindowsMixedRealityHMD/Private/WindowsMixedRealityHMD.cpp
--- a/Engine/Plugins/Runtime/WindowsMixedReality/Source/WindowsMixedRealityHMD/Private/WindowsMixedRealityHMD.cpp
+++ b/Engine/Plugins/Runtime/WindowsMixedReality/Source/WindowsMixedRealityHMD/Private/WindowsMixedRealityHMD.cpp
@@ -1096,8 +1096,6 @@ namespace WindowsMixedReality
 			outTargetableTexture,
 			outShaderResourceTexture);
 
-		CurrentBackBuffer = outTargetableTexture;
-
 		return true;
 	}
 
@@ -1184,8 +1182,17 @@ namespace WindowsMixedReality
 #if WITH_WINDOWS_MIXED_REALITY
 		// Update depth texture to match format Windows Mixed Reality platform is expecting.
 		FSceneRenderTargets& SceneContext = FSceneRenderTargets::Get(RHICmdList);
-		FRHITexture2D* depthFRHITexture = SceneContext.GetSceneDepthTexture().GetReference()->GetTexture2D();
+		if (SceneContext.SceneDepthZ == nullptr)
+		{
+			return;
+		}
+		const FTexture2DRHIRef SceneDepthTexture = SceneContext.GetSceneDepthTexture();
+		if (SceneDepthTexture == nullptr)
+		{
+			return;
+		}
 
+		FRHITexture2D* depthFRHITexture = SceneContext.GetSceneDepthTexture().GetReference()->GetTexture2D();
 		if (depthFRHITexture == nullptr || depthFRHITexture->GetNativeResource() == nullptr)
 		{
 			return;
diff --git a/Engine/Plugins/Runtime/WindowsMixedReality/Source/WindowsMixedRealityHMD/Private/WindowsMixedRealityHMD.h b/Engine/Plugins/Runtime/WindowsMixedReality/Source/WindowsMixedRealityHMD/Private/WindowsMixedRealityHMD.h
--- a/Engine/Plugins/Runtime/WindowsMixedReality/Source/WindowsMixedRealityHMD/Private/WindowsMixedRealityHMD.h
+++ b/Engine/Plugins/Runtime/WindowsMixedReality/Source/WindowsMixedRealityHMD/Private/WindowsMixedRealityHMD.h
@@ -172,8 +172,6 @@ namespace WindowsMixedReality
 		ID3D11Texture2D* stereoDepthTexture = nullptr;
 		const float farPlaneDistance = 650.0f;
 
-		// The back buffer for this frame
-		FTexture2DRHIRef CurrentBackBuffer;
 		void InitTrackingFrame();
 		TRefCountPtr<FWindowsMixedRealityCustomPresent> mCustomPresent = nullptr;
 
