Hololens Emulator QRCodesTrackerPlugin.dll related crash.
-Hololens Emulator was using the arm64 dll resulting in crashes.  Lets use the x64 ones instead.
-Discovered while reproducing the no audio on emulator bug.

diff --git a/Engine/Plugins/Runtime/AR/Microsoft/HoloLensAR/Source/HoloLensAR/HoloLensAR.Build.cs b/Engine/Plugins/Runtime/AR/Microsoft/HoloLensAR/Source/HoloLensAR/HoloLensAR.Build.cs
--- a/Engine/Plugins/Runtime/AR/Microsoft/HoloLensAR/Source/HoloLensAR/HoloLensAR.Build.cs
+++ b/Engine/Plugins/Runtime/AR/Microsoft/HoloLensAR/Source/HoloLensAR/HoloLensAR.Build.cs
@@ -31,13 +31,13 @@ public class HoloLensAR : ModuleRules
 			}
             );
 
-        if (Target.Platform == UnrealTargetPlatform.Win64)
+        if (Target.Platform == UnrealTargetPlatform.Win64 || (Target.Platform == UnrealTargetPlatform.HoloLens && Target.WindowsPlatform.Architecture == WindowsArchitecture.x64))
         {
             PublicDelayLoadDLLs.Add("QRCodesTrackerPlugin.dll");
             RuntimeDependencies.Add(System.IO.Path.Combine("$(EngineDir)/Binaries/ThirdParty/Windows/x64", "QRCodesTrackerPlugin.dll"));
             RuntimeDependencies.Add(System.IO.Path.Combine("$(EngineDir)/Binaries/ThirdParty/Windows/x64", "QRCodesTrackerPlugin.winmd"));
         }
-        if (Target.Platform == UnrealTargetPlatform.HoloLens)
+        else if (Target.Platform == UnrealTargetPlatform.HoloLens)
         {
             PublicDelayLoadDLLs.Add("QRCodesTrackerPlugin.dll");
             RuntimeDependencies.Add(System.IO.Path.Combine("$(EngineDir)/Binaries/ThirdParty/HoloLens/ARM64", "QRCodesTrackerPlugin.dll"));
diff --git a/Engine/Plugins/Runtime/WindowsMixedReality/Source/WindowsMixedRealityHMD/WindowsMixedRealityHMD.Build.cs b/Engine/Plugins/Runtime/WindowsMixedReality/Source/WindowsMixedRealityHMD/WindowsMixedRealityHMD.Build.cs
--- a/Engine/Plugins/Runtime/WindowsMixedReality/Source/WindowsMixedRealityHMD/WindowsMixedRealityHMD.Build.cs
+++ b/Engine/Plugins/Runtime/WindowsMixedReality/Source/WindowsMixedRealityHMD/WindowsMixedRealityHMD.Build.cs
@@ -133,14 +133,14 @@ namespace UnrealBuildTool.Rules
 				PrivatePCHHeaderFile = "Private/WindowsMixedRealityPrecompiled.h";
 			}
 
-            if (Target.Platform == UnrealTargetPlatform.Win64)
+            if (Target.Platform == UnrealTargetPlatform.Win64 || (Target.Platform == UnrealTargetPlatform.HoloLens && Target.WindowsPlatform.Architecture == WindowsArchitecture.x64))
             {
                 PublicDelayLoadDLLs.Add("QRCodesTrackerPlugin.dll");
                 RuntimeDependencies.Add(Path.Combine("$(EngineDir)/Binaries/ThirdParty/Windows/x64", "QRCodesTrackerPlugin.dll"));
             }
-            if (Target.Platform == UnrealTargetPlatform.HoloLens)
+            else if (Target.Platform == UnrealTargetPlatform.HoloLens)
             {
-                PublicDelayLoadDLLs.Add("QRCodesTrackerPlugin.dll");
+				PublicDelayLoadDLLs.Add("QRCodesTrackerPlugin.dll");
                 RuntimeDependencies.Add(Path.Combine("$(EngineDir)/Binaries/ThirdParty/HoloLens/ARM64", "QRCodesTrackerPlugin.dll"));
             }
 
