UE-79885 HL2 CL integrations
-Completed fix for "IsTracked" with HL2 hand tracking.

diff --git a/Engine/Plugins/Runtime/WindowsMixedReality/Source/WindowsMixedRealityHandTracking/Private/WindowsMixedRealityHandTracking.cpp b/Engine/Plugins/Runtime/WindowsMixedReality/Source/WindowsMixedRealityHandTracking/Private/WindowsMixedRealityHandTracking.cpp
--- a/Engine/Plugins/Runtime/WindowsMixedReality/Source/WindowsMixedRealityHandTracking/Private/WindowsMixedRealityHandTracking.cpp
+++ b/Engine/Plugins/Runtime/WindowsMixedReality/Source/WindowsMixedRealityHandTracking/Private/WindowsMixedRealityHandTracking.cpp
@@ -129,7 +129,6 @@ FLiveLinkSubjectName FWindowsMixedRealityHandTracking::LiveLinkRightHandTracking
 FWindowsMixedRealityHandTracking::FWindowsMixedRealityHandTracking(const TSharedRef<FGenericApplicationMessageHandler>& InMessageHandler)
 	: MessageHandler(InMessageHandler)
 	, DeviceIndex(0)
-	, bIsHandTrackingStateValid(true)
 {
 	// Register "MotionController" modular feature manually
 	IModularFeatures::Get().RegisterModularFeature(GetModularFeatureName(), this);
@@ -182,19 +181,19 @@ const FTransform& FWindowsMixedRealityHandTracking::FHandState::GetTransform(EWM
 
 bool FWindowsMixedRealityHandTracking::GetControllerOrientationAndPosition(const int32 ControllerIndex, const FName MotionSource, FRotator& OutOrientation, FVector& OutPosition, float WorldToMetersScale) const
 {
-	bool bFound = false;
+	bool bTracked = false;
 	if (ControllerIndex == DeviceIndex)
 	{
 		FTransform ControllerTransform = FTransform::Identity;
 		if (MotionSource == FName("Left"))
 		{
 			ControllerTransform = GetLeftHandState().GetTransform(EWMRHandKeypoint::Palm);
-			bFound = true;
+			bTracked = GetLeftHandState().ReceivedJointPoses;
 		}
 		else if (MotionSource == FName("Right"))
 		{
 			ControllerTransform = GetRightHandState().GetTransform(EWMRHandKeypoint::Palm);
-			bFound = true;
+			bTracked = GetRightHandState().ReceivedJointPoses;
 		}
 
 		// This can only be done in the game thread since it uses the UEnum directly
@@ -223,12 +222,13 @@ bool FWindowsMixedRealityHandTracking::GetControllerOrientationAndPosition(const
 				if (bUseRightHand)
 				{
 					ControllerTransform = GetRightHandState().GetTransform((EWMRHandKeypoint)ValueFromName);
+					bTracked = GetRightHandState().ReceivedJointPoses;
 				}
 				else
 				{
 					ControllerTransform = GetLeftHandState().GetTransform((EWMRHandKeypoint)ValueFromName);
+					bTracked = GetLeftHandState().ReceivedJointPoses;
 				}
-				bFound = true;
 			}
 		}
 
@@ -236,15 +236,10 @@ bool FWindowsMixedRealityHandTracking::GetControllerOrientationAndPosition(const
 		OutOrientation = ControllerTransform.GetRotation().Rotator();
 	}
 
-	if (bFound)
-	{
-		return true;
-	}
-	
 	// Then call super to handle a few of the default labels, for backward compatibility
 	FXRMotionControllerBase::GetControllerOrientationAndPosition(ControllerIndex, MotionSource, OutOrientation, OutPosition, WorldToMetersScale);
 
-	return false;
+	return bTracked;
 }
 
 bool FWindowsMixedRealityHandTracking::GetControllerOrientationAndPosition(const int32 ControllerIndex, const EControllerHand DeviceHand, FRotator& OutOrientation, FVector& OutPosition, float WorldToMetersScale) const
@@ -280,7 +275,9 @@ bool FWindowsMixedRealityHandTracking::GetControllerOrientationAndPosition(const
 
 ETrackingStatus FWindowsMixedRealityHandTracking::GetControllerTrackingStatus(const int32 ControllerIndex, const EControllerHand DeviceHand) const
 {
-	return bIsHandTrackingStateValid ? ETrackingStatus::Tracked : ETrackingStatus::NotTracked;
+	const FWindowsMixedRealityHandTracking::FHandState& HandState = (DeviceHand == EControllerHand::Left) ? GetLeftHandState() : GetRightHandState();
+
+	return HandState.ReceivedJointPoses ? ETrackingStatus::Tracked : ETrackingStatus::NotTracked;
 }
 
 FName FWindowsMixedRealityHandTracking::GetMotionControllerDeviceTypeName() const
@@ -312,7 +309,6 @@ void FWindowsMixedRealityHandTracking::SendControllerEvents()
 {
 // @TODO: implement for WMRSDK
 // #if WITH_MLSDK
-// 	if (bIsHandTrackingStateValid)
 // 	{
 // 		const MLHandTrackingData& CurrentHandTrackingData = GetCurrentHandTrackingData();
 // 		const MLHandTrackingData& OldHandTrackingData = GetPreviousHandTrackingData();
@@ -350,7 +346,7 @@ const FWindowsMixedRealityHandTracking::FHandState& FWindowsMixedRealityHandTrac
 
 bool FWindowsMixedRealityHandTracking::IsHandTrackingStateValid() const
 {
-	return bIsHandTrackingStateValid;
+	return true;
 }
 
 bool FWindowsMixedRealityHandTracking::GetKeypointTransform(EControllerHand Hand, EWMRHandKeypoint Keypoint, FTransform& OutTransform) const
@@ -366,7 +362,7 @@ void FWindowsMixedRealityHandTracking::UpdateTrackerData()
 	// Pump the interop update of the hand.
 	WindowsMixedReality::FWindowsMixedRealityStatics::PollHandTracking();
 
-	if (bIsHandTrackingStateValid && IWindowsMixedRealityHMDPlugin::Get().IsAvailable())
+	if (IWindowsMixedRealityHMDPlugin::Get().IsAvailable())
 	{
 		// Get all the bones for each hand
 		for (int32 Hand = 0; Hand < 2; Hand++)
diff --git a/Engine/Plugins/Runtime/WindowsMixedReality/Source/WindowsMixedRealityHandTracking/Private/WindowsMixedRealityHandTracking.h b/Engine/Plugins/Runtime/WindowsMixedReality/Source/WindowsMixedRealityHandTracking/Private/WindowsMixedRealityHandTracking.h
--- a/Engine/Plugins/Runtime/WindowsMixedReality/Source/WindowsMixedRealityHandTracking/Private/WindowsMixedRealityHandTracking.h
+++ b/Engine/Plugins/Runtime/WindowsMixedReality/Source/WindowsMixedRealityHandTracking/Private/WindowsMixedRealityHandTracking.h
@@ -97,8 +97,6 @@ private:
 
 	FHandState HandStates[2];
 
-	bool bIsHandTrackingStateValid;
-
 	// LiveLink Data
 	/** The local client to push data updates to */
 	ILiveLinkClient* LiveLinkClient = nullptr;
