Change: 8926470
UE-79885 HL2 CL integrations
Merging //UE4/Partner-Microsoft-XR to Release-4.23 (//UE4/Release-4.23) Change: 8926470
-Update GetHandJointTransform to world space.

diff --git a/Engine/Plugins/Runtime/WindowsMixedReality/Source/WindowsMixedRealityHandTracking/Private/WindowsMixedRealityHandTrackingFunctionLibrary.cpp b/Engine/Plugins/Runtime/WindowsMixedReality/Source/WindowsMixedRealityHandTracking/Private/WindowsMixedRealityHandTrackingFunctionLibrary.cpp
--- a/Engine/Plugins/Runtime/WindowsMixedReality/Source/WindowsMixedRealityHandTracking/Private/WindowsMixedRealityHandTrackingFunctionLibrary.cpp
+++ b/Engine/Plugins/Runtime/WindowsMixedReality/Source/WindowsMixedRealityHandTracking/Private/WindowsMixedRealityHandTrackingFunctionLibrary.cpp
@@ -22,7 +22,7 @@ bool UWindowsMixedRealityHandTrackingFunctionLibrary::GetHandJointTransform(ECon
 
 		if (bSuccess)
 		{
-			OutTransform = KeyPointTransform;
+			OutTransform = KeyPointTransform * UHeadMountedDisplayFunctionLibrary::GetTrackingToWorldTransform(GWorld);
 			return true;
 		}
 	}
