Undef LOCTEXT_NAMESPACE
#rnx

diff --git a/Engine/Plugins/Runtime/WindowsMixedReality/Source/WindowsMixedRealityHMD/Private/WindowsMixedRealityFunctionLibrary.cpp b/Engine/Plugins/Runtime/WindowsMixedReality/Source/WindowsMixedRealityHMD/Private/WindowsMixedRealityFunctionLibrary.cpp
--- a/Engine/Plugins/Runtime/WindowsMixedReality/Source/WindowsMixedRealityHMD/Private/WindowsMixedRealityFunctionLibrary.cpp
+++ b/Engine/Plugins/Runtime/WindowsMixedReality/Source/WindowsMixedRealityHMD/Private/WindowsMixedRealityFunctionLibrary.cpp
@@ -186,4 +186,5 @@ static FAutoConsoleCommand CSetFocusPointForFrameCmd(
 	TEXT("vr.SetFocusPointForFrame"),
 	*LOCTEXT("CVarText_SetFocusPointForFrame",
 		"Set the reference point for the stabilization plane on hololens 2. You must set it each frame to activate the feature for that frame.").ToString(),
-	FConsoleCommandWithWorldArgsAndOutputDeviceDelegate::CreateStatic(SetFocusPointForFrame));
\ No newline at end of file
+	FConsoleCommandWithWorldArgsAndOutputDeviceDelegate::CreateStatic(SetFocusPointForFrame));
+#undef LOCTEXT_NAMESPACE
\ No newline at end of file
