#DEV-VR - Fix for non-unity build issue with windows header

diff --git a/Engine/Plugins/Runtime/WindowsMixedReality/Source/WindowsMixedRealityHMD/Classes/WindowsMixedRealityFunctionLibrary.h b/Engine/Plugins/Runtime/WindowsMixedReality/Source/WindowsMixedRealityHMD/Classes/WindowsMixedRealityFunctionLibrary.h
--- a/Engine/Plugins/Runtime/WindowsMixedReality/Source/WindowsMixedRealityHMD/Classes/WindowsMixedRealityFunctionLibrary.h
+++ b/Engine/Plugins/Runtime/WindowsMixedReality/Source/WindowsMixedRealityHMD/Classes/WindowsMixedRealityFunctionLibrary.h
@@ -3,6 +3,7 @@
 #pragma once
 
 #include "CoreMinimal.h"
+#include "Windows/WindowsHWrapper.h"
 #include "Kismet/BlueprintFunctionLibrary.h"
 
 #include "MixedRealityInterop.h"
diff --git a/Engine/Plugins/Runtime/WindowsMixedReality/Source/WindowsMixedRealityHMD/Classes/WindowsMixedRealityStatics.h b/Engine/Plugins/Runtime/WindowsMixedReality/Source/WindowsMixedRealityHMD/Classes/WindowsMixedRealityStatics.h
--- a/Engine/Plugins/Runtime/WindowsMixedReality/Source/WindowsMixedRealityHMD/Classes/WindowsMixedRealityStatics.h
+++ b/Engine/Plugins/Runtime/WindowsMixedReality/Source/WindowsMixedRealityHMD/Classes/WindowsMixedRealityStatics.h
@@ -3,6 +3,7 @@
 #pragma once
 
 #include "CoreMinimal.h"
+#include "Windows/WindowsHWrapper.h"
 
 #pragma warning(disable:4668)  
 #include <DirectXMath.h>
