Attempt to get Windows Mixed Reality to be skipped on Linux

diff --git a/Engine/Plugins/Runtime/WindowsMixedReality/WindowsMixedReality.uplugin b/Engine/Plugins/Runtime/WindowsMixedReality/WindowsMixedReality.uplugin
--- a/Engine/Plugins/Runtime/WindowsMixedReality/WindowsMixedReality.uplugin
+++ b/Engine/Plugins/Runtime/WindowsMixedReality/WindowsMixedReality.uplugin
@@ -14,6 +14,10 @@
 	"CanContainContent" : false,
 	"IsBetaVersion" : true,
 	"Installed" : false,
+	"SupportedTargetPlatforms": [
+		"Win32",
+		"Win64"
+	],
 	"Modules": [
 		{
 			"Name": "WindowsMixedRealityHMD",
