Guarding against GEngine dereference to fix static analysis

diff --git a/Engine/Plugins/Runtime/WindowsMixedReality/Source/WindowsMixedRealityHMD/Private/WindowsMixedRealityHMD.cpp b/Engine/Plugins/Runtime/WindowsMixedReality/Source/WindowsMixedRealityHMD/Private/WindowsMixedRealityHMD.cpp
--- a/Engine/Plugins/Runtime/WindowsMixedReality/Source/WindowsMixedRealityHMD/Private/WindowsMixedRealityHMD.cpp
+++ b/Engine/Plugins/Runtime/WindowsMixedReality/Source/WindowsMixedRealityHMD/Private/WindowsMixedRealityHMD.cpp
@@ -321,7 +321,7 @@ namespace WindowsMixedReality
 		}
 
 		// Restore windows focus to game window to preserve keyboard/mouse input.
-		if (currentWornState == EHMDWornState::Type::Worn)
+		if ((currentWornState == EHMDWornState::Type::Worn) && GEngine)
 		{
 			HWND gameHWND = (HWND)GEngine->GameViewport->GetWindow()->GetNativeWindow()->GetOSWindowHandle();
 
