Moving WindowsMixedReality RuntimeSettings and PlatformEditor under the Plugin directory

diff --git a/Engine/Plugins/Runtime/WindowsMixedReality/Source/WindowsMixedRealityPlatformEditor/Private/WindowsMixedRealityDetails.cpp b/Engine/Plugins/Runtime/WindowsMixedReality/Source/WindowsMixedRealityPlatformEditor/Private/WindowsMixedRealityDetails.cpp
new file mode 100644
--- /dev/null
+++ b/Engine/Plugins/Runtime/WindowsMixedReality/Source/WindowsMixedRealityPlatformEditor/Private/WindowsMixedRealityDetails.cpp
@@ -0,0 +1,63 @@
+// Copyright (c) Microsoft Corporation. All rights reserved.
+
+#include "WindowsMixedRealityDetails.h"
+#include "DetailLayoutBuilder.h"
+#include "DetailCategoryBuilder.h"
+#include "IDetailPropertyRow.h"
+#include "DetailWidgetRow.h"
+#include "IDetailGroup.h"
+#include "PropertyHandle.h"
+#include "PropertyCustomizationHelpers.h"
+#include "Widgets/Input/SButton.h"
+
+#include "WindowsMixedRealityRuntimeSettings.h"
+#include "WindowsMixedRealityStatics.h"
+
+#define LOCTEXT_NAMESPACE "FWindowsMixedRealityDetails"
+
+TSharedRef<IDetailCustomization> FWindowsMixedRealityDetails::MakeInstance()
+{
+	return MakeShareable(new FWindowsMixedRealityDetails);
+}
+
+void FWindowsMixedRealityDetails::CustomizeDetails(IDetailLayoutBuilder& DetailBuilder)
+{
+	IDetailCategoryBuilder& remotingCategory = DetailBuilder.EditCategory(TEXT("Holographic Remoting"));
+	remotingCategory.AddCustomRow(LOCTEXT("Connect Button", "Connect Button"))
+		[
+			SNew(SButton)
+			.Text(LOCTEXT("Connect", "Connect"))
+			.OnClicked_Raw(this, &FWindowsMixedRealityDetails::OnConnectButtonClicked)
+		];
+	
+	remotingCategory.AddCustomRow(LOCTEXT("Disconnect Button", "Disconnect Button"))
+		[
+			SNew(SButton)
+			.Text(LOCTEXT("Disconnect", "Disconnect"))
+			.OnClicked_Raw(this, &FWindowsMixedRealityDetails::OnDisconnectButtonClicked)
+		];
+}
+
+FReply FWindowsMixedRealityDetails::OnConnectButtonClicked()
+{
+	UWindowsMixedRealityRuntimeSettings* settings = UWindowsMixedRealityRuntimeSettings::Get();
+
+	FString ip = settings->RemoteHoloLensIP;
+	UE_LOG(LogTemp, Log, TEXT("Connecting to: %s"), *ip);
+
+	unsigned int bitrate = settings->MaxBitrate;
+
+	WindowsMixedReality::FWindowsMixedRealityStatics::ConnectToRemoteHoloLens(ip, bitrate);
+
+	return FReply::Handled();
+}
+
+FReply FWindowsMixedRealityDetails::OnDisconnectButtonClicked()
+{
+	UE_LOG(LogTemp, Log, TEXT("Disconnecting from remote HoloLens"));
+	WindowsMixedReality::FWindowsMixedRealityStatics::DisconnectFromRemoteHoloLens();
+
+	return FReply::Handled();
+}
+
+#undef LOCTEXT_NAMESPACE
diff --git a/Engine/Plugins/Runtime/WindowsMixedReality/Source/WindowsMixedRealityPlatformEditor/Private/WindowsMixedRealityDetails.h b/Engine/Plugins/Runtime/WindowsMixedReality/Source/WindowsMixedRealityPlatformEditor/Private/WindowsMixedRealityDetails.h
new file mode 100644
--- /dev/null
+++ b/Engine/Plugins/Runtime/WindowsMixedReality/Source/WindowsMixedRealityPlatformEditor/Private/WindowsMixedRealityDetails.h
@@ -0,0 +1,20 @@
+// Copyright (c) Microsoft Corporation. All rights reserved.
+
+#pragma once
+
+#include "IDetailCustomNodeBuilder.h"
+#include "PropertyHandle.h"
+#include "IDetailCustomization.h"
+#include "IPropertyTypeCustomization.h"
+
+class FWindowsMixedRealityDetails : public IDetailCustomization
+{
+public:
+	static TSharedRef<IDetailCustomization> MakeInstance();
+
+	virtual void CustomizeDetails(IDetailLayoutBuilder& DetailBuilder) override;
+
+private:
+	FReply OnConnectButtonClicked();
+	FReply OnDisconnectButtonClicked();
+};
\ No newline at end of file
diff --git a/Engine/Plugins/Runtime/WindowsMixedReality/Source/WindowsMixedRealityPlatformEditor/Private/WindowsMixedRealityPlatformEditorModule.cpp b/Engine/Plugins/Runtime/WindowsMixedReality/Source/WindowsMixedRealityPlatformEditor/Private/WindowsMixedRealityPlatformEditorModule.cpp
new file mode 100644
--- /dev/null
+++ b/Engine/Plugins/Runtime/WindowsMixedReality/Source/WindowsMixedRealityPlatformEditor/Private/WindowsMixedRealityPlatformEditorModule.cpp
@@ -0,0 +1,54 @@
+// Copyright (c) Microsoft Corporation. All rights reserved.
+
+#include "WindowsMixedRealityRuntimeSettings.h"
+#include "Modules/ModuleInterface.h"
+#include "ISettingsModule.h"
+#include "Modules/ModuleManager.h"
+#include "PropertyEditorModule.h"
+#include "WindowsMixedRealityDetails.h"
+
+#define LOCTEXT_NAMESPACE "FWindowsMixedRealityPlatformEditorModule"
+
+
+/**
+ * Module for WindowsMR platform editor utilities
+ */
+class FWindowsMixedRealityPlatformEditorModule
+	: public IModuleInterface
+{
+	virtual void StartupModule() override
+	{
+		FPropertyEditorModule& PropertyModule = FModuleManager::LoadModuleChecked<FPropertyEditorModule>("PropertyEditor");
+		PropertyModule.RegisterCustomClassLayout(FName("WindowsMixedRealityRuntimeSettings"), FOnGetDetailCustomizationInstance::CreateStatic(&FWindowsMixedRealityDetails::MakeInstance));
+		PropertyModule.NotifyCustomizationModuleChanged();
+
+		// register settings
+		ISettingsModule* SettingsModule = FModuleManager::GetModulePtr<ISettingsModule>("Settings");
+
+		if (SettingsModule != nullptr)
+		{
+			{
+				SettingsModule->RegisterSettings("Project", "Platforms", "WindowsMixedReality",
+					LOCTEXT("RuntimeSettingsName", "Windows Mixed Reality"),
+					LOCTEXT("RuntimeSettingsDescription", "Project settings for Windows Mixed Reality"),
+					GetMutableDefault<UWindowsMixedRealityRuntimeSettings>()
+				);
+			}
+		}
+	}
+
+	virtual void ShutdownModule() override
+	{
+		ISettingsModule* SettingsModule = FModuleManager::GetModulePtr<ISettingsModule>("Settings");
+
+		if (SettingsModule != nullptr)
+		{
+			SettingsModule->UnregisterSettings("Project", "Platforms", "WindowsMixedReality");
+		}
+	}
+};
+
+
+IMPLEMENT_MODULE(FWindowsMixedRealityPlatformEditorModule, WindowsMixedRealityPlatformEditor);
+
+#undef LOCTEXT_NAMESPACE
diff --git a/Engine/Plugins/Runtime/WindowsMixedReality/Source/WindowsMixedRealityPlatformEditor/WindowsMixedRealityPlatformEditor.Build.cs b/Engine/Plugins/Runtime/WindowsMixedReality/Source/WindowsMixedRealityPlatformEditor/WindowsMixedRealityPlatformEditor.Build.cs
new file mode 100644
--- /dev/null
+++ b/Engine/Plugins/Runtime/WindowsMixedReality/Source/WindowsMixedRealityPlatformEditor/WindowsMixedRealityPlatformEditor.Build.cs
@@ -0,0 +1,35 @@
+// Copyright (c) Microsoft Corporation. All rights reserved.
+
+using UnrealBuildTool;
+
+public class WindowsMixedRealityPlatformEditor : ModuleRules
+{
+	public WindowsMixedRealityPlatformEditor(ReadOnlyTargetRules Target) : base(Target)
+	{
+		PrivateDependencyModuleNames.AddRange(
+			new string[] {
+				"Core",
+				"CoreUObject",
+				"InputCore",
+				"Engine",
+				"Slate",
+				"SlateCore",
+				"EditorStyle",
+				"EditorWidgets",
+				"DesktopWidgets",
+				"PropertyEditor",
+				"SharedSettingsWidgets",
+				"WindowsMixedRealityRuntimeSettings",
+				"WindowsMixedRealityHMD",
+				"TargetPlatform",
+				"RenderCore",
+			}
+		);
+
+		PrivateIncludePathModuleNames.AddRange(
+			new string[] {
+				"Settings"
+			}
+		);
+	}
+}
diff --git a/Engine/Plugins/Runtime/WindowsMixedReality/Source/WindowsMixedRealityRuntimeSettings/Classes/WindowsMixedRealityRuntimeSettings.h b/Engine/Plugins/Runtime/WindowsMixedReality/Source/WindowsMixedRealityRuntimeSettings/Classes/WindowsMixedRealityRuntimeSettings.h
new file mode 100644
--- /dev/null
+++ b/Engine/Plugins/Runtime/WindowsMixedReality/Source/WindowsMixedRealityRuntimeSettings/Classes/WindowsMixedRealityRuntimeSettings.h
@@ -0,0 +1,40 @@
+// Copyright (c) Microsoft Corporation. All rights reserved.
+
+#pragma once
+
+#include "CoreMinimal.h"
+#include "UObject/ObjectMacros.h"
+#include "UObject/Object.h"
+#include "Engine/EngineTypes.h"
+#include "WindowsMixedRealityRuntimeSettings.generated.h"
+
+/**
+* Add a default value for every new UPROPERTY in this class in <UnrealEngine>/Engine/Config/BaseEngine.ini
+*/
+
+/**
+ * Implements the settings for the WindowsMixedReality runtime platform.
+ */
+UCLASS(config=Engine, defaultconfig)
+class WINDOWSMIXEDREALITYRUNTIMESETTINGS_API UWindowsMixedRealityRuntimeSettings : public UObject
+{
+public:
+	GENERATED_BODY()
+
+private:
+	static class UWindowsMixedRealityRuntimeSettings* WMRSettingsSingleton;
+
+public:
+	static UWindowsMixedRealityRuntimeSettings* Get();
+
+#if WITH_EDITOR
+	virtual void PostEditChangeProperty(struct FPropertyChangedEvent& PropertyChangedEvent) override;
+#endif // WITH_EDITOR
+
+	/** The IP of the HoloLens to remote to. */
+	UPROPERTY(GlobalConfig, EditAnywhere, Category = "Holographic Remoting", Meta = (DisplayName = "IP of HoloLens to remote to."))
+	FString RemoteHoloLensIP;
+
+	UPROPERTY(GlobalConfig, EditAnywhere, Category = "Holographic Remoting", Meta = (DisplayName = "Max network transfer rate (kb/s)"))
+	unsigned int MaxBitrate = 4000;
+};
diff --git a/Engine/Plugins/Runtime/WindowsMixedReality/Source/WindowsMixedRealityRuntimeSettings/Private/WindowsMixedRealityRuntimeSettings.cpp b/Engine/Plugins/Runtime/WindowsMixedReality/Source/WindowsMixedRealityRuntimeSettings/Private/WindowsMixedRealityRuntimeSettings.cpp
new file mode 100644
--- /dev/null
+++ b/Engine/Plugins/Runtime/WindowsMixedReality/Source/WindowsMixedRealityRuntimeSettings/Private/WindowsMixedRealityRuntimeSettings.cpp
@@ -0,0 +1,34 @@
+// Copyright (c) Microsoft Corporation. All rights reserved.
+
+#include "WindowsMixedRealityRuntimeSettings.h"
+#include "Misc/ConfigCacheIni.h"
+#include "CoreGlobals.h"
+#include "UObject/Package.h"
+
+UWindowsMixedRealityRuntimeSettings* UWindowsMixedRealityRuntimeSettings::WMRSettingsSingleton = nullptr;
+
+#if WITH_EDITOR
+
+void UWindowsMixedRealityRuntimeSettings::PostEditChangeProperty(struct FPropertyChangedEvent& PropertyChangedEvent)
+{
+	GConfig->Flush(1);
+}
+
+#endif
+
+UWindowsMixedRealityRuntimeSettings* UWindowsMixedRealityRuntimeSettings::Get()
+{
+	if (WMRSettingsSingleton == nullptr)
+	{
+		static const TCHAR* SettingsContainerName = TEXT("WindowsMixedRealityRuntimeSettingsContainer");
+
+		WMRSettingsSingleton = FindObject<UWindowsMixedRealityRuntimeSettings>(GetTransientPackage(), SettingsContainerName);
+
+		if (WMRSettingsSingleton == nullptr)
+		{
+			WMRSettingsSingleton = NewObject<UWindowsMixedRealityRuntimeSettings>(GetTransientPackage(), UWindowsMixedRealityRuntimeSettings::StaticClass(), SettingsContainerName);
+			WMRSettingsSingleton->AddToRoot();
+		}
+	}
+	return WMRSettingsSingleton;
+}
\ No newline at end of file
diff --git a/Engine/Plugins/Runtime/WindowsMixedReality/Source/WindowsMixedRealityRuntimeSettings/Private/WindowsMixedRealityRuntimeSettingsModule.cpp b/Engine/Plugins/Runtime/WindowsMixedReality/Source/WindowsMixedRealityRuntimeSettings/Private/WindowsMixedRealityRuntimeSettingsModule.cpp
new file mode 100644
--- /dev/null
+++ b/Engine/Plugins/Runtime/WindowsMixedReality/Source/WindowsMixedRealityRuntimeSettings/Private/WindowsMixedRealityRuntimeSettingsModule.cpp
@@ -0,0 +1,6 @@
+// Copyright (c) Microsoft Corporation. All rights reserved.
+
+#include "Modules/ModuleInterface.h"
+#include "Modules/ModuleManager.h"
+
+IMPLEMENT_MODULE(FDefaultModuleImpl, WindowsMixedRealityRuntimeSettings);
diff --git a/Engine/Plugins/Runtime/WindowsMixedReality/Source/WindowsMixedRealityRuntimeSettings/WindowsMixedRealityRuntimeSettings.Build.cs b/Engine/Plugins/Runtime/WindowsMixedReality/Source/WindowsMixedRealityRuntimeSettings/WindowsMixedRealityRuntimeSettings.Build.cs
new file mode 100644
--- /dev/null
+++ b/Engine/Plugins/Runtime/WindowsMixedReality/Source/WindowsMixedRealityRuntimeSettings/WindowsMixedRealityRuntimeSettings.Build.cs
@@ -0,0 +1,28 @@
+// Copyright (c) Microsoft Corporation. All rights reserved.
+
+using UnrealBuildTool;
+
+public class WindowsMixedRealityRuntimeSettings : ModuleRules
+{
+	public WindowsMixedRealityRuntimeSettings(ReadOnlyTargetRules Target) : base(Target)
+	{	
+		PrivateDependencyModuleNames.AddRange(
+			new string[]
+			{
+				"Core",
+				"CoreUObject",
+				"Engine",
+			}
+		);
+
+		if (Target.Type == TargetRules.TargetType.Editor || Target.Type == TargetRules.TargetType.Program)
+		{
+			PrivateDependencyModuleNames.AddRange(
+				new string[]
+				{
+					"TargetPlatform",
+				}
+			);
+		}
+	}
+}
diff --git a/Engine/Plugins/Runtime/WindowsMixedReality/WindowsMixedReality.uplugin b/Engine/Plugins/Runtime/WindowsMixedReality/WindowsMixedReality.uplugin
--- a/Engine/Plugins/Runtime/WindowsMixedReality/WindowsMixedReality.uplugin
+++ b/Engine/Plugins/Runtime/WindowsMixedReality/WindowsMixedReality.uplugin
@@ -19,6 +19,15 @@
 		"Win64"
 	],
 	"Modules": [
+		{
+			"Name": "WindowsMixedRealityRuntimeSettings",
+			"Type": "Runtime",
+			"LoadingPhase": "PostConfigInit",
+			"WhitelistPlatforms": [
+				"Win32",
+				"Win64"
+			]
+		},
 		{
 			"Name": "WindowsMixedRealityHMD",
 			"Type": "Runtime",
@@ -36,6 +45,15 @@
 				"Win32",
 				"Win64"
 			]
-		}
+		},
+		{
+			"Name" : "WindowsMixedRealityPlatformEditor",
+			"Type" : "Editor",
+			"LoadingPhase" : "PostEngineInit",
+			"WhitelistPlatforms": [
+				"Win32",
+				"Win64"
+			]
+		},
 	]
 }
diff --git a/Engine/Source/Developer/Windows/WindowsMixedRealityPlatformEditor/Private/WindowsMixedRealityDetails.cpp b/Engine/Source/Developer/Windows/WindowsMixedRealityPlatformEditor/Private/WindowsMixedRealityDetails.cpp
deleted file mode 100644
diff --git a/Engine/Source/Developer/Windows/WindowsMixedRealityPlatformEditor/Private/WindowsMixedRealityDetails.h b/Engine/Source/Developer/Windows/WindowsMixedRealityPlatformEditor/Private/WindowsMixedRealityDetails.h
deleted file mode 100644
diff --git a/Engine/Source/Developer/Windows/WindowsMixedRealityPlatformEditor/Private/WindowsMixedRealityPlatformEditorModule.cpp b/Engine/Source/Developer/Windows/WindowsMixedRealityPlatformEditor/Private/WindowsMixedRealityPlatformEditorModule.cpp
deleted file mode 100644
diff --git a/Engine/Source/Developer/Windows/WindowsMixedRealityPlatformEditor/WindowsMixedRealityPlatformEditor.Build.cs b/Engine/Source/Developer/Windows/WindowsMixedRealityPlatformEditor/WindowsMixedRealityPlatformEditor.Build.cs
deleted file mode 100644
diff --git a/Engine/Source/Runtime/WindowsMixedReality/WindowsMixedRealityRuntimeSettings/Classes/WindowsMixedRealityRuntimeSettings.h b/Engine/Source/Runtime/WindowsMixedReality/WindowsMixedRealityRuntimeSettings/Classes/WindowsMixedRealityRuntimeSettings.h
deleted file mode 100644
diff --git a/Engine/Source/Runtime/WindowsMixedReality/WindowsMixedRealityRuntimeSettings/Private/WindowsMixedRealityRuntimeSettings.cpp b/Engine/Source/Runtime/WindowsMixedReality/WindowsMixedRealityRuntimeSettings/Private/WindowsMixedRealityRuntimeSettings.cpp
deleted file mode 100644
diff --git a/Engine/Source/Runtime/WindowsMixedReality/WindowsMixedRealityRuntimeSettings/Private/WindowsMixedRealityRuntimeSettingsModule.cpp b/Engine/Source/Runtime/WindowsMixedReality/WindowsMixedRealityRuntimeSettings/Private/WindowsMixedRealityRuntimeSettingsModule.cpp
deleted file mode 100644
diff --git a/Engine/Source/Runtime/WindowsMixedReality/WindowsMixedRealityRuntimeSettings/WindowsMixedRealityRuntimeSettings.Build.cs b/Engine/Source/Runtime/WindowsMixedReality/WindowsMixedRealityRuntimeSettings/WindowsMixedRealityRuntimeSettings.Build.cs
deleted file mode 100644
