Guarding Windows Mixed Reality against old version of the compiler

diff --git a/Engine/Plugins/Runtime/WindowsMixedReality/Source/WindowsMixedRealityHMD/Classes/WindowsMixedRealityFunctionLibrary.h b/Engine/Plugins/Runtime/WindowsMixedReality/Source/WindowsMixedRealityHMD/Classes/WindowsMixedRealityFunctionLibrary.h
--- a/Engine/Plugins/Runtime/WindowsMixedReality/Source/WindowsMixedRealityHMD/Classes/WindowsMixedRealityFunctionLibrary.h
+++ b/Engine/Plugins/Runtime/WindowsMixedReality/Source/WindowsMixedRealityHMD/Classes/WindowsMixedRealityFunctionLibrary.h
@@ -6,7 +6,9 @@
 #include "Windows/WindowsHWrapper.h"
 #include "Kismet/BlueprintFunctionLibrary.h"
 
+#if WITH_WINDOWS_MIXED_REALITY
 #include "MixedRealityInterop.h"
+#endif
 
 #include "WindowsMixedRealityFunctionLibrary.Generated.h"
 
diff --git a/Engine/Plugins/Runtime/WindowsMixedReality/Source/WindowsMixedRealityHMD/Classes/WindowsMixedRealityStatics.h b/Engine/Plugins/Runtime/WindowsMixedReality/Source/WindowsMixedRealityHMD/Classes/WindowsMixedRealityStatics.h
--- a/Engine/Plugins/Runtime/WindowsMixedReality/Source/WindowsMixedRealityHMD/Classes/WindowsMixedRealityStatics.h
+++ b/Engine/Plugins/Runtime/WindowsMixedReality/Source/WindowsMixedRealityHMD/Classes/WindowsMixedRealityStatics.h
@@ -9,7 +9,9 @@
 #include <DirectXMath.h>
 #pragma warning(default:4668)
 
+#if WITH_WINDOWS_MIXED_REALITY
 #include "MixedRealityInterop.h"
+#endif
 
 namespace WindowsMixedReality
 {
@@ -51,6 +53,7 @@ namespace WindowsMixedReality
 
 		static bool SupportsSpatialInput();
 
+#if WITH_WINDOWS_MIXED_REALITY
 		static MixedRealityInterop::HMDTrackingStatus GetControllerTrackingStatus(MixedRealityInterop::HMDHand hand);
 
 		static bool GetControllerOrientationAndPosition(MixedRealityInterop::HMDHand hand, FRotator& OutOrientation, FVector& OutPosition);
@@ -68,7 +71,7 @@ namespace WindowsMixedReality
 		static void SubmitHapticValue(
 			MixedRealityInterop::HMDHand hand,
 			float value);
-
+#endif
 		// Remoting
 		static void ConnectToRemoteHoloLens(FString remoteIP, unsigned int bitrate);
 		static void DisconnectFromRemoteHoloLens();
diff --git a/Engine/Plugins/Runtime/WindowsMixedReality/Source/WindowsMixedRealityHMD/Private/WindowsMixedRealityCustomPresent.h b/Engine/Plugins/Runtime/WindowsMixedReality/Source/WindowsMixedRealityHMD/Private/WindowsMixedRealityCustomPresent.h
--- a/Engine/Plugins/Runtime/WindowsMixedReality/Source/WindowsMixedRealityHMD/Private/WindowsMixedRealityCustomPresent.h
+++ b/Engine/Plugins/Runtime/WindowsMixedReality/Source/WindowsMixedRealityHMD/Private/WindowsMixedRealityCustomPresent.h
@@ -8,16 +8,24 @@
 #include "D3D11RHIPrivate.h"
 #include "D3D11Util.h"
 
+#if WITH_WINDOWS_MIXED_REALITY
 #include "MixedRealityInterop.h"
+#endif
 
 namespace WindowsMixedReality
 {
 	class FWindowsMixedRealityCustomPresent : public FRHICustomPresent
 	{
 	public:
-		FWindowsMixedRealityCustomPresent(MixedRealityInterop* _hmd, ID3D11Device* device)
+		FWindowsMixedRealityCustomPresent(
+#if WITH_WINDOWS_MIXED_REALITY
+			MixedRealityInterop* _hmd, 
+#endif
+			ID3D11Device* device)
 			: FRHICustomPresent()
+#if WITH_WINDOWS_MIXED_REALITY
 			, hmd(_hmd)
+#endif
 		{
 			// Get the D3D11 context.
 			device->GetImmediateContext(&D3D11Context);
@@ -31,6 +39,7 @@ namespace WindowsMixedReality
 		}
 		virtual bool Present(int32 & InOutSyncInterval) override
 		{
+#if WITH_WINDOWS_MIXED_REALITY
 			if (hmd == nullptr ||
 				D3D11Context == nullptr ||
 				ViewportTexture == nullptr)
@@ -39,6 +48,9 @@ namespace WindowsMixedReality
 			}
 
 			return hmd->Present(D3D11Context, ViewportTexture);
+#else
+			return false;
+#endif
 		}
 
 		void UpdateViewport(
@@ -66,7 +78,9 @@ namespace WindowsMixedReality
 
 
 	private:
+#if WITH_WINDOWS_MIXED_REALITY
 		MixedRealityInterop* hmd = nullptr;
+#endif
 
 		ID3D11DeviceContext* D3D11Context = nullptr;
 		ID3D11Texture2D* ViewportTexture = nullptr;
diff --git a/Engine/Plugins/Runtime/WindowsMixedReality/Source/WindowsMixedRealityHMD/Private/WindowsMixedRealityHMD.cpp b/Engine/Plugins/Runtime/WindowsMixedReality/Source/WindowsMixedRealityHMD/Private/WindowsMixedRealityHMD.cpp
--- a/Engine/Plugins/Runtime/WindowsMixedReality/Source/WindowsMixedRealityHMD/Private/WindowsMixedRealityHMD.cpp
+++ b/Engine/Plugins/Runtime/WindowsMixedReality/Source/WindowsMixedRealityHMD/Private/WindowsMixedRealityHMD.cpp
@@ -20,6 +20,7 @@
 // Windows Mixed Reality HMD Plugin
 //---------------------------------------------------
 
+#if WITH_WINDOWS_MIXED_REALITY
 class FDepthConversionPS : public FGlobalShader
 {
 	DECLARE_SHADER_TYPE(FDepthConversionPS, Global);
@@ -71,10 +72,13 @@ private:
 };
 
 IMPLEMENT_SHADER_TYPE(, FDepthConversionPS, TEXT("/Plugin/WindowsMixedReality/Private/DepthConversion.usf"), TEXT("MainPixelShader"), SF_Pixel)
+#endif
 
 namespace WindowsMixedReality
 {
+#if WITH_WINDOWS_MIXED_REALITY
 	MixedRealityInterop hmd;
+#endif
 
 	class FWindowsMixedRealityHMDPlugin : public IWindowsMixedRealityHMDPlugin
 	{
@@ -83,7 +87,11 @@ namespace WindowsMixedReality
 
 		bool IsHMDConnected()
 		{
+#if WITH_WINDOWS_MIXED_REALITY
 			return WindowsMixedReality::hmd.IsAvailable();
+#else
+			return false;
+#endif
 		}
 
 		FString GetModuleKeyName() const override
@@ -93,19 +101,29 @@ namespace WindowsMixedReality
 
 		void StartupModule() override
 		{
+#if WITH_WINDOWS_MIXED_REALITY
 			IHeadMountedDisplayModule::StartupModule();
 			FString PluginShaderDir = FPaths::Combine(IPluginManager::Get().FindPlugin(TEXT("WindowsMixedReality"))->GetBaseDir(), TEXT("Shaders"));
 			AddShaderSourceDirectoryMapping(TEXT("/Plugin/WindowsMixedReality"), PluginShaderDir);
+#else
+			UE_LOG(LogCore, Error, TEXT("Windows Mixed Reality compiled with unsupported compiler.  Please recompile with Visual Studio 2017"));
+#endif
 		}
 
 		void ShutdownModule() override
 		{
+#if WITH_WINDOWS_MIXED_REALITY
 			hmd.Dispose(true);
+#endif
 		}
 
 		uint64 GetGraphicsAdapterLuid() override
 		{
+#if WITH_WINDOWS_MIXED_REALITY
 			return hmd.GraphicsAdapterLUID();
+#else
+			return 0;
+#endif
 		}
 	};
 
@@ -130,12 +148,16 @@ namespace WindowsMixedReality
 
 	bool FWindowsMixedRealityHMD::IsHMDConnected()
 	{
+#if WITH_WINDOWS_MIXED_REALITY
 		if (hmd.IsRemoting())
 		{
 			return true;
 		}
 
 		return hmd.IsAvailable();
+#else
+		return false;
+#endif
 	}
 
 	bool FWindowsMixedRealityHMD::IsHMDEnabled() const
@@ -145,6 +167,7 @@ namespace WindowsMixedReality
 
 	EHMDWornState::Type FWindowsMixedRealityHMD::GetHMDWornState()
 	{
+#if WITH_WINDOWS_MIXED_REALITY
 		if (hmd.IsRemoting())
 		{
 			return EHMDWornState::Type::Unknown;
@@ -165,6 +188,9 @@ namespace WindowsMixedReality
 		};
 
 		return wornState;
+#else
+		return EHMDWornState::Unknown;
+#endif
 	}
 
 	void FWindowsMixedRealityHMD::OnBeginPlay(FWorldContext & InWorldContext)
@@ -239,7 +265,11 @@ namespace WindowsMixedReality
 
 	FString FWindowsMixedRealityHMD::GetVersionString() const
 	{
+#if WITH_WINDOWS_MIXED_REALITY
 		return FString(hmd.GetDisplayName());
+#else
+		return FString();
+#endif
 	}
 
 	void CenterMouse(RECT windowRect)
@@ -262,6 +292,7 @@ namespace WindowsMixedReality
 			return true;
 		}
 
+#if WITH_WINDOWS_MIXED_REALITY
 		if (!hmd.IsInitialized())
 		{
 			D3D11Device = InternalGetD3D11Device();
@@ -348,6 +379,7 @@ namespace WindowsMixedReality
 		}
 
 		CachedWorldToMetersScale = WorldContext.World()->GetWorldSettings()->WorldToMeters;
+#endif
 
 		return true;
 	}
@@ -423,11 +455,14 @@ namespace WindowsMixedReality
 
 	void FWindowsMixedRealityHMD::ResetOrientationAndPosition(float yaw)
 	{
+#if WITH_WINDOWS_MIXED_REALITY
 		hmd.ResetOrientationAndPosition();
+#endif
 	}
 
 	void FWindowsMixedRealityHMD::InitTrackingFrame()
 	{
+#if WITH_WINDOWS_MIXED_REALITY
 		DirectX::XMMATRIX leftPose;
 		DirectX::XMMATRIX rightPose;
 		MixedRealityInterop::HMDTrackingOrigin trackingOrigin;
@@ -469,10 +504,12 @@ namespace WindowsMixedReality
 			CurrOrientation = HeadRotation;
 			CurrPosition = HeadPosition;
 		}
+#endif
 	}
 
 	void SetupHiddenVisibleAreaMesh(TArray<FHMDViewMesh>& HiddenMeshes, TArray<FHMDViewMesh>& VisibleMeshes)
 	{
+#if WITH_WINDOWS_MIXED_REALITY
 		for (int i = (int)MixedRealityInterop::HMDEye::Left;
 			i <= (int)MixedRealityInterop::HMDEye::Right; i++)
 		{
@@ -512,10 +549,12 @@ namespace WindowsMixedReality
 				delete[] vertexPositions;
 			}
 		}
+#endif
 	}
 
 	void FWindowsMixedRealityHMD::SetupHolographicCamera()
 	{
+#if WITH_WINDOWS_MIXED_REALITY
 		// Set the viewport to match the HMD display
 		FSceneViewport* SceneVP = FindMRSceneViewport(bIsStereoDesired);
 
@@ -566,6 +605,7 @@ namespace WindowsMixedReality
 
 		// Uncap fps to enable FPS higher than 62
 		GEngine->bForceDisableFrameRateSmoothing = bIsStereoEnabled;
+#endif
 	}
 
 	bool FWindowsMixedRealityHMD::IsStereoEnabled() const
@@ -575,6 +615,7 @@ namespace WindowsMixedReality
 
 	bool FWindowsMixedRealityHMD::EnableStereo(bool stereo)
 	{
+#if WITH_WINDOWS_MIXED_REALITY
 		if (stereo)
 		{
 			if (bIsStereoDesired && hmd.IsInitialized())
@@ -604,7 +645,7 @@ namespace WindowsMixedReality
 			FApp::SetUseVRFocus(false);
 			FApp::SetHasVRFocus(false);
 		}
-
+#endif
 		return bIsStereoDesired;
 	}
 
@@ -659,6 +700,7 @@ namespace WindowsMixedReality
 			return FMatrix::Identity;
 		}
 
+#if WITH_WINDOWS_MIXED_REALITY
 		DirectX::XMMATRIX projection = (StereoPassType == eSSP_LEFT_EYE)
 			? hmd.GetProjectionMatrix(MixedRealityInterop::HMDEye::Left)
 			: hmd.GetProjectionMatrix(MixedRealityInterop::HMDEye::Right);
@@ -677,6 +719,9 @@ namespace WindowsMixedReality
 		result.M[3][2] = Nz;
 
 		return result;
+#else
+		return FMatrix::Identity;
+#endif
 	}
 
 	void FWindowsMixedRealityHMD::GetEyeRenderParams_RenderThread(
@@ -705,7 +750,12 @@ namespace WindowsMixedReality
 	FIntPoint FWindowsMixedRealityHMD::GetIdealRenderTargetSize() const
 	{
 		int Width, Height;
+#if WITH_WINDOWS_MIXED_REALITY
 		hmd.GetDisplayDimensions(Width, Height);
+#else
+		Width = 100;
+		Height = 100;
+#endif
 
 		return FIntPoint(Width * 2, Height);
 	}
@@ -727,11 +777,13 @@ namespace WindowsMixedReality
 	// Called when screen size changes.
 	void FWindowsMixedRealityHMD::UpdateViewportRHIBridge(bool bUseSeparateRenderTarget, const class FViewport& Viewport, FRHIViewport* const ViewportRHI)
 	{
+#if WITH_WINDOWS_MIXED_REALITY
 		if (IsStereoEnabled() && mCustomPresent != nullptr)
 		{
 			hmd.SetScreenScaleFactor(ScreenScalePercentage);
 			mCustomPresent->UpdateViewport(Viewport, ViewportRHI);
 		}
+#endif
 	}
 
 	void FWindowsMixedRealityHMD::RenderTexture_RenderThread(class FRHICommandListImmediate& RHICmdList, class FRHITexture2D* BackBuffer, class FRHITexture2D* SrcTexture, FVector2D WindowSize) const
@@ -901,6 +953,7 @@ namespace WindowsMixedReality
 	{
 		check(IsInRenderingThread());
 
+#if WITH_WINDOWS_MIXED_REALITY
 		// Update depth texture to match format Windows Mixed Reality platform is expecting.
 		FSceneRenderTargets& SceneContext = FSceneRenderTargets::Get(RHICmdList);
 		FRHITexture2D* depthFRHITexture = SceneContext.GetSceneDepthTexture().GetReference()->GetTexture2D();
@@ -943,7 +996,6 @@ namespace WindowsMixedReality
 
 		TShaderMapRef<FScreenVS> vertexShader(shaderMap);
 		TShaderMapRef<FDepthConversionPS> pixelShader(shaderMap);
-
 		GraphicsPSOInit.BoundShaderState.VertexDeclarationRHI = RendererModule->GetFilterVertexDeclaration().VertexDeclarationRHI;
 		GraphicsPSOInit.BoundShaderState.VertexShaderRHI = GETSAFERHISHADER_VERTEX(*vertexShader);
 		GraphicsPSOInit.BoundShaderState.PixelShaderRHI = GETSAFERHISHADER_PIXEL(*pixelShader);
@@ -992,12 +1044,14 @@ namespace WindowsMixedReality
 		}
 
 		stereoDepthTexture = (ID3D11Texture2D*)remappedDepthTexture->GetNativeResource();
+#endif
 	}
 
 	void FWindowsMixedRealityHMD::PreRenderViewFamily_RenderThread(
 		FRHICommandListImmediate& RHICmdList,
 		FSceneViewFamily& InViewFamily)
 	{
+#if WITH_WINDOWS_MIXED_REALITY
 		if (!mCustomPresent || !hmd.IsInitialized() || !hmd.IsAvailable())
 		{
 			return;
@@ -1010,6 +1064,7 @@ namespace WindowsMixedReality
 			// Because Windows Mixed Reality can only have 2 rendering parameters in flight at any time, this is fatal.
 			this->bRequestRestart = true;
 		}
+#endif
 	}
 
 	bool FWindowsMixedRealityHMD::IsActiveThisFrame(class FViewport* InViewport) const
@@ -1049,8 +1104,9 @@ namespace WindowsMixedReality
 	{
 		check(IsInGameThread());
 
+#if WITH_WINDOWS_MIXED_REALITY
 		hmd.EnableStereo(false);
-
+#endif
 		// Ensure that we aren't currently trying to render a frame before destroying our custom present.
 		FlushRenderingCommands();
 		StopCustomPresent();
@@ -1084,12 +1140,17 @@ namespace WindowsMixedReality
 
 	bool FWindowsMixedRealityHMD::IsCurrentlyImmersive()
 	{
+#if WITH_WINDOWS_MIXED_REALITY
 		return hmd.IsCurrentlyImmersive();
+#else
+		return false;
+#endif
 	}
 
 	// Setup Windows Holographic view and tracking space.
 	void FWindowsMixedRealityHMD::InitializeHolographic()
 	{
+#if WITH_WINDOWS_MIXED_REALITY
 		if (!hmd.IsInitialized())
 		{
 			D3D11Device = InternalGetD3D11Device();
@@ -1109,6 +1170,7 @@ namespace WindowsMixedReality
 		{
 			PauseHandle = FCoreDelegates::ApplicationWillEnterBackgroundDelegate.AddRaw(this, &FWindowsMixedRealityHMD::AppServicePause);
 		}
+#endif
 	}
 
 	// Prevent crashes if computer goes to sleep.
@@ -1119,16 +1181,22 @@ namespace WindowsMixedReality
 
 	bool WindowsMixedReality::FWindowsMixedRealityHMD::IsAvailable()
 	{
+#if WITH_WINDOWS_MIXED_REALITY
 		return hmd.IsAvailable();
+#else
+		return false;
+#endif
 	}
 
 	// Initialize Windows Holographic present.
 	void FWindowsMixedRealityHMD::StartCustomPresent()
 	{
+#if WITH_WINDOWS_MIXED_REALITY
 		if (mCustomPresent == nullptr)
 		{
 			mCustomPresent = new FWindowsMixedRealityCustomPresent(&hmd, D3D11Device);
 		}
+#endif
 	}
 
 	// Cleanup resources for holographic present.
@@ -1140,9 +1208,14 @@ namespace WindowsMixedReality
 	// Spatial Input
 	bool FWindowsMixedRealityHMD::SupportsSpatialInput()
 	{
+#if WITH_WINDOWS_MIXED_REALITY
 		return hmd.SupportsSpatialInput();
+#else
+		return false;
+#endif
 	}
 
+#if WITH_WINDOWS_MIXED_REALITY
 	MixedRealityInterop::HMDTrackingStatus FWindowsMixedRealityHMD::GetControllerTrackingStatus(MixedRealityInterop::HMDHand hand)
 	{
 		return hmd.GetControllerTrackingStatus(hand);
@@ -1172,7 +1245,6 @@ namespace WindowsMixedReality
 
 			return true;
 		}
-
 		return false;
 	}
 
@@ -1201,6 +1273,7 @@ namespace WindowsMixedReality
 	{
 		hmd.SubmitHapticValue(hand, value);
 	}
+#endif
 
 	// Remoting
 	void FWindowsMixedRealityHMD::ConnectToRemoteHoloLens(const wchar_t* ip, unsigned int bitrate)
@@ -1208,14 +1281,18 @@ namespace WindowsMixedReality
 #if WITH_EDITOR
 		D3D11Device = InternalGetD3D11Device();
 
+#if WITH_WINDOWS_MIXED_REALITY
 		hmd.ConnectToRemoteHoloLens(D3D11Device.GetReference(), ip, bitrate);
+#endif
 #endif
 	}
 
 	void FWindowsMixedRealityHMD::DisconnectFromRemoteHoloLens()
 	{
 #if WITH_EDITOR
+#if WITH_WINDOWS_MIXED_REALITY
 		hmd.DisconnectFromRemoteHoloLens();
+#endif
 #endif
 	}
 }
diff --git a/Engine/Plugins/Runtime/WindowsMixedReality/Source/WindowsMixedRealityHMD/Private/WindowsMixedRealityHMD.h b/Engine/Plugins/Runtime/WindowsMixedReality/Source/WindowsMixedRealityHMD/Private/WindowsMixedRealityHMD.h
--- a/Engine/Plugins/Runtime/WindowsMixedReality/Source/WindowsMixedRealityHMD/Private/WindowsMixedRealityHMD.h
+++ b/Engine/Plugins/Runtime/WindowsMixedReality/Source/WindowsMixedRealityHMD/Private/WindowsMixedRealityHMD.h
@@ -16,7 +16,9 @@
 #include "XRRenderTargetManager.h"
 #include "RendererInterface.h"
 
+#if WITH_WINDOWS_MIXED_REALITY
 #include "MixedRealityInterop.h"
+#endif
 
 namespace WindowsMixedReality
 {
@@ -227,9 +229,11 @@ namespace WindowsMixedReality
 		// Spatial input
 		bool IsAvailable();
 		bool SupportsSpatialInput();
+#if WITH_WINDOWS_MIXED_REALITY
 		static MixedRealityInterop::HMDTrackingStatus GetControllerTrackingStatus(MixedRealityInterop::HMDHand hand);
 		bool GetControllerOrientationAndPosition(MixedRealityInterop::HMDHand hand, FRotator & OutOrientation, FVector & OutPosition);
 		bool PollInput();
+
 		MixedRealityInterop::HMDInputPressState GetPressState(
 			MixedRealityInterop::HMDHand hand,
 			MixedRealityInterop::HMDInputControllerButtons button);
@@ -239,6 +243,7 @@ namespace WindowsMixedReality
 		void SubmitHapticValue(
 			MixedRealityInterop::HMDHand hand,
 			float value);
+#endif
 		void LockMouseToCenter(bool locked)
 		{
 			mouseLockedToCenter = locked;
diff --git a/Engine/Plugins/Runtime/WindowsMixedReality/Source/WindowsMixedRealityHMD/Private/WindowsMixedRealityStatics.cpp b/Engine/Plugins/Runtime/WindowsMixedReality/Source/WindowsMixedRealityHMD/Private/WindowsMixedRealityStatics.cpp
--- a/Engine/Plugins/Runtime/WindowsMixedReality/Source/WindowsMixedRealityHMD/Private/WindowsMixedRealityStatics.cpp
+++ b/Engine/Plugins/Runtime/WindowsMixedReality/Source/WindowsMixedRealityHMD/Private/WindowsMixedRealityStatics.cpp
@@ -29,6 +29,7 @@ namespace WindowsMixedReality
 		return false;
 	}
 
+#if WITH_WINDOWS_MIXED_REALITY
 	MixedRealityInterop::HMDTrackingStatus FWindowsMixedRealityStatics::GetControllerTrackingStatus(MixedRealityInterop::HMDHand hand)
 	{
 		FWindowsMixedRealityHMD* hmd = GetWindowsMixedRealityHMD();
@@ -100,6 +101,7 @@ namespace WindowsMixedReality
 			hmd->SubmitHapticValue(hand, value);
 		}
 	}
+#endif
 
 	// Remoting
 	void FWindowsMixedRealityStatics::ConnectToRemoteHoloLens(FString remoteIP, unsigned int bitrate)
diff --git a/Engine/Plugins/Runtime/WindowsMixedReality/Source/WindowsMixedRealityHMD/WindowsMixedRealityHMD.Build.cs b/Engine/Plugins/Runtime/WindowsMixedReality/Source/WindowsMixedRealityHMD/WindowsMixedRealityHMD.Build.cs
--- a/Engine/Plugins/Runtime/WindowsMixedReality/Source/WindowsMixedRealityHMD/WindowsMixedRealityHMD.Build.cs
+++ b/Engine/Plugins/Runtime/WindowsMixedReality/Source/WindowsMixedRealityHMD/WindowsMixedRealityHMD.Build.cs
@@ -3,6 +3,10 @@
 using System;
 using System.IO;
 using UnrealBuildTool;
+using System.Collections.Generic;
+using Microsoft.Win32;
+using System.Diagnostics;
+
 
 namespace UnrealBuildTool.Rules
 {
@@ -20,34 +24,42 @@ namespace UnrealBuildTool.Rules
 
 		private void LoadMixedReality(ReadOnlyTargetRules Target)
         {
-			string LibrariesPath = Path.Combine(ThirdPartyPath, "Lib", "x64");
+            int releaseId = Convert.ToInt32(Registry.GetValue(@"HKEY_LOCAL_MACHINE\SOFTWARE\Microsoft\Windows NT\CurrentVersion", "ReleaseId", "").ToString());
+            bool bAllowWindowsMixedReality = (Target.WindowsPlatform.Compiler != UnrealBuildTool.WindowsCompiler.VisualStudio2015) && (releaseId >= 1803);
 
-            if (Target.Platform == UnrealTargetPlatform.Win32)
-            {
-                LibrariesPath = Path.Combine(ThirdPartyPath, "Lib", "x86");
-                RuntimeDependencies.Add("$(EngineDir)/Binaries/Win32/HolographicStreamerDesktop.dll");
-                RuntimeDependencies.Add("$(EngineDir)/Binaries/Win32/Microsoft.Perception.Simulation.dll");
-                RuntimeDependencies.Add("$(EngineDir)/Binaries/Win32/PerceptionSimulationManager.dll");
-            }
-            else if (Target.Platform == UnrealTargetPlatform.Win64)
+            if (bAllowWindowsMixedReality)
             {
-                RuntimeDependencies.Add("$(EngineDir)/Binaries/Win64/HolographicStreamerDesktop.dll");
-                RuntimeDependencies.Add("$(EngineDir)/Binaries/Win64/Microsoft.Perception.Simulation.dll");
-                RuntimeDependencies.Add("$(EngineDir)/Binaries/Win64/PerceptionSimulationManager.dll");
+                string LibrariesPath = Path.Combine(ThirdPartyPath, "Lib", "x64");
+
+                if (Target.Platform == UnrealTargetPlatform.Win32)
+                {
+                    LibrariesPath = Path.Combine(ThirdPartyPath, "Lib", "x86");
+                    RuntimeDependencies.Add("$(EngineDir)/Binaries/Win32/HolographicStreamerDesktop.dll");
+                    RuntimeDependencies.Add("$(EngineDir)/Binaries/Win32/Microsoft.Perception.Simulation.dll");
+                    RuntimeDependencies.Add("$(EngineDir)/Binaries/Win32/PerceptionSimulationManager.dll");
+                }
+                else if (Target.Platform == UnrealTargetPlatform.Win64)
+                {
+                    RuntimeDependencies.Add("$(EngineDir)/Binaries/Win64/HolographicStreamerDesktop.dll");
+                    RuntimeDependencies.Add("$(EngineDir)/Binaries/Win64/Microsoft.Perception.Simulation.dll");
+                    RuntimeDependencies.Add("$(EngineDir)/Binaries/Win64/PerceptionSimulationManager.dll");
+                }
+
+                PublicLibraryPaths.Add(LibrariesPath);
+                PublicAdditionalLibraries.Add(Path.Combine(LibrariesPath, "MixedRealityInterop.lib"));
+                PublicAdditionalLibraries.Add(Path.Combine(LibrariesPath, "HolographicStreamerDesktop.lib"));
+                PublicAdditionalLibraries.Add(Path.Combine(LibrariesPath, "Microsoft.Perception.Simulation.lib"));
+                PublicAdditionalLibraries.Add(Path.Combine(LibrariesPath, "PerceptionSimulationManager.lib"));
+
+                // Win10 support
+                PublicAdditionalLibraries.Add(Path.Combine(LibrariesPath, "onecore.lib"));
+                // Explicitly load lib path since name conflicts with an existing lib in the DX11 dependency.
+                PublicAdditionalLibraries.Add(Path.Combine(LibrariesPath, "d3d11.lib"));
+
+                PublicIncludePaths.Add(Path.Combine(ThirdPartyPath, "Include"));
             }
 
-            PublicLibraryPaths.Add(LibrariesPath);
-            PublicAdditionalLibraries.Add(Path.Combine(LibrariesPath, "MixedRealityInterop.lib"));
-            PublicAdditionalLibraries.Add(Path.Combine(LibrariesPath, "HolographicStreamerDesktop.lib"));
-            PublicAdditionalLibraries.Add(Path.Combine(LibrariesPath, "Microsoft.Perception.Simulation.lib"));
-            PublicAdditionalLibraries.Add(Path.Combine(LibrariesPath, "PerceptionSimulationManager.lib"));
-
-            // Win10 support
-            PublicAdditionalLibraries.Add(Path.Combine(LibrariesPath, "onecore.lib"));
-			// Explicitly load lib path since name conflicts with an existing lib in the DX11 dependency.
-            PublicAdditionalLibraries.Add(Path.Combine(LibrariesPath, "d3d11.lib"));
-				
-			PublicIncludePaths.Add(Path.Combine(ThirdPartyPath, "Include"));
+            PublicDefinitions.Add("WITH_WINDOWS_MIXED_REALITY=" + (bAllowWindowsMixedReality ? "1" : "0"));
         }
 
         public WindowsMixedRealityHMD(ReadOnlyTargetRules Target) : base(Target)
diff --git a/Engine/Plugins/Runtime/WindowsMixedReality/Source/WindowsMixedRealitySpatialInput/Private/WindowsMixedRealitySpatialInput.cpp b/Engine/Plugins/Runtime/WindowsMixedReality/Source/WindowsMixedRealitySpatialInput/Private/WindowsMixedRealitySpatialInput.cpp
--- a/Engine/Plugins/Runtime/WindowsMixedReality/Source/WindowsMixedRealitySpatialInput/Private/WindowsMixedRealitySpatialInput.cpp
+++ b/Engine/Plugins/Runtime/WindowsMixedReality/Source/WindowsMixedRealitySpatialInput/Private/WindowsMixedRealitySpatialInput.cpp
@@ -44,16 +44,18 @@ namespace WindowsMixedReality
 
 	void FWindowsMixedRealitySpatialInput::SendControllerEvents()
 	{
+#if WITH_WINDOWS_MIXED_REALITY
 		if (!FWindowsMixedRealityStatics::PollInput())
 		{
 			return;
 		}
 
 		const uint32 sourceId = 0;
-		SendButtonEvents(sourceId);
 		SendAxisEvents(sourceId);
+#endif
 	}
 
+#if WITH_WINDOWS_MIXED_REALITY
 	void SendControllerButtonEvent(
 		TSharedPtr< FGenericApplicationMessageHandler > messageHandler,
 		uint32 controllerId,
@@ -85,6 +87,7 @@ namespace WindowsMixedReality
 				false);
 		}
 	}
+#endif
 
 	void SendControllerAxisEvent(
 		TSharedPtr< FGenericApplicationMessageHandler > messageHandler,
@@ -100,6 +103,7 @@ namespace WindowsMixedReality
 			static_cast<float>(axisPosition));
 	}
 
+#if WITH_WINDOWS_MIXED_REALITY
 	void FWindowsMixedRealitySpatialInput::SendAxisEvents(uint32 source)
 	{
 		FKey key;
@@ -240,6 +244,7 @@ namespace WindowsMixedReality
 			}
 		}
 	}
+#endif
 
 	void FWindowsMixedRealitySpatialInput::SetMessageHandler(const TSharedRef<FGenericApplicationMessageHandler>& InMessageHandler)
 	{
@@ -300,8 +305,10 @@ namespace WindowsMixedReality
 			return;
 		}
 
+#if WITH_WINDOWS_MIXED_REALITY
 		MixedRealityInterop::HMDHand hand = (MixedRealityInterop::HMDHand)DeviceHand;
 		FWindowsMixedRealityStatics::SubmitHapticValue(hand, (Values.Frequency > 0.0f) ? Values.Amplitude : 0.0f);
+#endif
 	}
 
 	void FWindowsMixedRealitySpatialInput::GetHapticFrequencyRange(float & MinFrequency, float & MaxFrequency) const
@@ -323,20 +330,28 @@ namespace WindowsMixedReality
 
 	bool FWindowsMixedRealitySpatialInput::GetControllerOrientationAndPosition(const int32 ControllerIndex, const EControllerHand DeviceHand, FRotator & OutOrientation, FVector & OutPosition, float WorldToMetersScale) const
 	{
+#if WITH_WINDOWS_MIXED_REALITY
 		MixedRealityInterop::HMDHand hand = (MixedRealityInterop::HMDHand)((int)DeviceHand);
 
 		bool success = FWindowsMixedRealityStatics::GetControllerOrientationAndPosition(hand, OutOrientation, OutPosition);
 		OutPosition *= WorldToMetersScale;
 
 		return success;
+#else
+		return false;
+#endif
 	}
 
 	ETrackingStatus FWindowsMixedRealitySpatialInput::GetControllerTrackingStatus(const int32 ControllerIndex, const EControllerHand DeviceHand) const
 	{
+#if WITH_WINDOWS_MIXED_REALITY
 		MixedRealityInterop::HMDHand hand = (MixedRealityInterop::HMDHand)((int)DeviceHand);
 		MixedRealityInterop::HMDTrackingStatus trackingStatus = FWindowsMixedRealityStatics::GetControllerTrackingStatus(hand);
 
 		return (ETrackingStatus)((int)trackingStatus);
+#else
+		return ETrackingStatus::NotTracked;
+#endif
 	}
 
 	void FWindowsMixedRealitySpatialInput::RegisterKeys() noexcept
diff --git a/Engine/Plugins/Runtime/WindowsMixedReality/Source/WindowsMixedRealitySpatialInput/Private/WindowsMixedRealitySpatialInput.h b/Engine/Plugins/Runtime/WindowsMixedReality/Source/WindowsMixedRealitySpatialInput/Private/WindowsMixedRealitySpatialInput.h
--- a/Engine/Plugins/Runtime/WindowsMixedReality/Source/WindowsMixedRealitySpatialInput/Private/WindowsMixedRealitySpatialInput.h
+++ b/Engine/Plugins/Runtime/WindowsMixedReality/Source/WindowsMixedRealitySpatialInput/Private/WindowsMixedRealitySpatialInput.h
@@ -57,8 +57,10 @@ namespace WindowsMixedReality
 		void InitializeSpatialInput() noexcept;
 		void UninitializeSpatialInput() noexcept;
 
+#if WITH_WINDOWS_MIXED_REALITY
 		void SendButtonEvents(uint32 source);
 		void SendAxisEvents(uint32 source);
+#endif
 
 		bool isLeftTouchpadTouched = false;
 		bool isRightTouchpadTouched = false;
diff --git a/Engine/Plugins/Runtime/WindowsMixedReality/Source/WindowsMixedRealitySpatialInput/WindowsMixedRealitySpatialInput.Build.cs b/Engine/Plugins/Runtime/WindowsMixedReality/Source/WindowsMixedRealitySpatialInput/WindowsMixedRealitySpatialInput.Build.cs
--- a/Engine/Plugins/Runtime/WindowsMixedReality/Source/WindowsMixedRealitySpatialInput/WindowsMixedRealitySpatialInput.Build.cs
+++ b/Engine/Plugins/Runtime/WindowsMixedReality/Source/WindowsMixedRealitySpatialInput/WindowsMixedRealitySpatialInput.Build.cs
@@ -3,6 +3,7 @@
 using System;
 using System.IO;
 using UnrealBuildTool;
+using Microsoft.Win32;
 
 namespace UnrealBuildTool.Rules
 {
@@ -20,32 +21,38 @@ namespace UnrealBuildTool.Rules
 		
 		private void LoadMixedReality(ReadOnlyTargetRules Target)
         {
-            string LibrariesPath = Path.Combine(ThirdPartyPath, "Lib", "x64");
+            int releaseId = Convert.ToInt32(Registry.GetValue(@"HKEY_LOCAL_MACHINE\SOFTWARE\Microsoft\Windows NT\CurrentVersion", "ReleaseId", "").ToString());
+            bool bAllowWindowsMixedReality = (Target.WindowsPlatform.Compiler != UnrealBuildTool.WindowsCompiler.VisualStudio2015) && (releaseId >= 1803);
 
-            if (Target.Platform == UnrealTargetPlatform.Win32)
+            if (bAllowWindowsMixedReality)
             {
-                LibrariesPath = Path.Combine(ThirdPartyPath, "Lib", "x86");
-                RuntimeDependencies.Add("$(EngineDir)/Binaries/Win32/HolographicStreamerDesktop.dll");
-                RuntimeDependencies.Add("$(EngineDir)/Binaries/Win32/Microsoft.Perception.Simulation.dll");
-                RuntimeDependencies.Add("$(EngineDir)/Binaries/Win32/PerceptionSimulationManager.dll");
-            }
-            else if (Target.Platform == UnrealTargetPlatform.Win64)
-            {
-                RuntimeDependencies.Add("$(EngineDir)/Binaries/Win64/HolographicStreamerDesktop.dll");
-                RuntimeDependencies.Add("$(EngineDir)/Binaries/Win64/Microsoft.Perception.Simulation.dll");
-                RuntimeDependencies.Add("$(EngineDir)/Binaries/Win64/PerceptionSimulationManager.dll");
-            }
+                string LibrariesPath = Path.Combine(ThirdPartyPath, "Lib", "x64");
 
-            PublicLibraryPaths.Add(LibrariesPath);
-            PublicAdditionalLibraries.Add(Path.Combine(LibrariesPath, "MixedRealityInterop.lib"));
-            PublicAdditionalLibraries.Add(Path.Combine(LibrariesPath, "HolographicStreamerDesktop.lib"));
-            PublicAdditionalLibraries.Add(Path.Combine(LibrariesPath, "Microsoft.Perception.Simulation.lib"));
-            PublicAdditionalLibraries.Add(Path.Combine(LibrariesPath, "PerceptionSimulationManager.lib"));
+                if (Target.Platform == UnrealTargetPlatform.Win32)
+                {
+                    LibrariesPath = Path.Combine(ThirdPartyPath, "Lib", "x86");
+                    RuntimeDependencies.Add("$(EngineDir)/Binaries/Win32/HolographicStreamerDesktop.dll");
+                    RuntimeDependencies.Add("$(EngineDir)/Binaries/Win32/Microsoft.Perception.Simulation.dll");
+                    RuntimeDependencies.Add("$(EngineDir)/Binaries/Win32/PerceptionSimulationManager.dll");
+                }
+                else if (Target.Platform == UnrealTargetPlatform.Win64)
+                {
+                    RuntimeDependencies.Add("$(EngineDir)/Binaries/Win64/HolographicStreamerDesktop.dll");
+                    RuntimeDependencies.Add("$(EngineDir)/Binaries/Win64/Microsoft.Perception.Simulation.dll");
+                    RuntimeDependencies.Add("$(EngineDir)/Binaries/Win64/PerceptionSimulationManager.dll");
+                }
 
-            // Win10 support
-            PublicAdditionalLibraries.Add(Path.Combine(LibrariesPath, "onecore.lib"));
+                PublicLibraryPaths.Add(LibrariesPath);
+                PublicAdditionalLibraries.Add(Path.Combine(LibrariesPath, "MixedRealityInterop.lib"));
+                PublicAdditionalLibraries.Add(Path.Combine(LibrariesPath, "HolographicStreamerDesktop.lib"));
+                PublicAdditionalLibraries.Add(Path.Combine(LibrariesPath, "Microsoft.Perception.Simulation.lib"));
+                PublicAdditionalLibraries.Add(Path.Combine(LibrariesPath, "PerceptionSimulationManager.lib"));
 
-            PublicIncludePaths.Add(Path.Combine(ThirdPartyPath, "Include"));
+                // Win10 support
+                PublicAdditionalLibraries.Add(Path.Combine(LibrariesPath, "onecore.lib"));
+
+                PublicIncludePaths.Add(Path.Combine(ThirdPartyPath, "Include"));
+            }
         }
 		
         public WindowsMixedRealitySpatialInput(ReadOnlyTargetRules Target) : base(Target)
