Guards around missing registry settings for the builder

diff --git a/Engine/Plugins/Runtime/WindowsMixedReality/Source/WindowsMixedRealityHMD/WindowsMixedRealityHMD.Build.cs b/Engine/Plugins/Runtime/WindowsMixedReality/Source/WindowsMixedRealityHMD/WindowsMixedRealityHMD.Build.cs
--- a/Engine/Plugins/Runtime/WindowsMixedReality/Source/WindowsMixedRealityHMD/WindowsMixedRealityHMD.Build.cs
+++ b/Engine/Plugins/Runtime/WindowsMixedReality/Source/WindowsMixedRealityHMD/WindowsMixedRealityHMD.Build.cs
@@ -24,7 +24,12 @@ namespace UnrealBuildTool.Rules
 
 		private void LoadMixedReality(ReadOnlyTargetRules Target)
         {
-            int releaseId = Convert.ToInt32(Registry.GetValue(@"HKEY_LOCAL_MACHINE\SOFTWARE\Microsoft\Windows NT\CurrentVersion", "ReleaseId", "").ToString());
+            int releaseId = 0;
+            string releaseIdString = Registry.GetValue(@"HKEY_LOCAL_MACHINE\SOFTWARE\Microsoft\Windows NT\CurrentVersion", "ReleaseId", "") as String;
+            if (!String.IsNullOrEmpty(releaseIdString))
+            {
+                releaseId = Convert.ToInt32(releaseIdString);
+            }
             bool bAllowWindowsMixedReality = (Target.WindowsPlatform.Compiler != UnrealBuildTool.WindowsCompiler.VisualStudio2015) && (releaseId >= 1803);
 
             if (bAllowWindowsMixedReality)
diff --git a/Engine/Plugins/Runtime/WindowsMixedReality/Source/WindowsMixedRealitySpatialInput/WindowsMixedRealitySpatialInput.Build.cs b/Engine/Plugins/Runtime/WindowsMixedReality/Source/WindowsMixedRealitySpatialInput/WindowsMixedRealitySpatialInput.Build.cs
--- a/Engine/Plugins/Runtime/WindowsMixedReality/Source/WindowsMixedRealitySpatialInput/WindowsMixedRealitySpatialInput.Build.cs
+++ b/Engine/Plugins/Runtime/WindowsMixedReality/Source/WindowsMixedRealitySpatialInput/WindowsMixedRealitySpatialInput.Build.cs
@@ -21,7 +21,12 @@ namespace UnrealBuildTool.Rules
 		
 		private void LoadMixedReality(ReadOnlyTargetRules Target)
         {
-            int releaseId = Convert.ToInt32(Registry.GetValue(@"HKEY_LOCAL_MACHINE\SOFTWARE\Microsoft\Windows NT\CurrentVersion", "ReleaseId", "").ToString());
+            int releaseId = 0;
+            string releaseIdString = Registry.GetValue(@"HKEY_LOCAL_MACHINE\SOFTWARE\Microsoft\Windows NT\CurrentVersion", "ReleaseId", "") as String;
+            if (!String.IsNullOrEmpty(releaseIdString))
+            {
+                releaseId = Convert.ToInt32(releaseIdString);
+            }
             bool bAllowWindowsMixedReality = (Target.WindowsPlatform.Compiler != UnrealBuildTool.WindowsCompiler.VisualStudio2015) && (releaseId >= 1803);
 
             if (bAllowWindowsMixedReality)
