Clamp the haptic value between 0 and 1 for WMR controllers.

diff --git a/Engine/Plugins/Runtime/WindowsMixedReality/Source/WindowsMixedRealityHMD/Private/WindowsMixedRealityHMD.cpp b/Engine/Plugins/Runtime/WindowsMixedReality/Source/WindowsMixedRealityHMD/Private/WindowsMixedRealityHMD.cpp
--- a/Engine/Plugins/Runtime/WindowsMixedReality/Source/WindowsMixedRealityHMD/Private/WindowsMixedRealityHMD.cpp
+++ b/Engine/Plugins/Runtime/WindowsMixedReality/Source/WindowsMixedRealityHMD/Private/WindowsMixedRealityHMD.cpp
@@ -1271,7 +1271,7 @@ namespace WindowsMixedReality
 
 	void FWindowsMixedRealityHMD::SubmitHapticValue(MixedRealityInterop::HMDHand hand, float value)
 	{
-		hmd.SubmitHapticValue(hand, value);
+		hmd.SubmitHapticValue(hand, FMath::Clamp(value, 0.f, 1.f));
 	}
 #endif
 
