Fix button events not firing on WMR controllers.
This seems to have been an accidental deletion in CL 4554826.

diff --git a/Engine/Plugins/Runtime/WindowsMixedReality/Source/WindowsMixedRealitySpatialInput/Private/WindowsMixedRealitySpatialInput.cpp b/Engine/Plugins/Runtime/WindowsMixedReality/Source/WindowsMixedRealitySpatialInput/Private/WindowsMixedRealitySpatialInput.cpp
--- a/Engine/Plugins/Runtime/WindowsMixedReality/Source/WindowsMixedRealitySpatialInput/Private/WindowsMixedRealitySpatialInput.cpp
+++ b/Engine/Plugins/Runtime/WindowsMixedReality/Source/WindowsMixedRealitySpatialInput/Private/WindowsMixedRealitySpatialInput.cpp
@@ -51,6 +51,7 @@ namespace WindowsMixedReality
 		}
 
 		const uint32 sourceId = 0;
+		SendButtonEvents(sourceId);
 		SendAxisEvents(sourceId);
 #endif
 	}
