UE-71585 VR Mode and VR Preview are able to be selected when the WMR plugin is enabled and there is no headset attached
-Created a runtime setting to enable/disable the creation of XRTracking system by WindowsMixedReality so that it does not appear availible in the editor because the user could use remoting to connect.
-This setting disables the plugin's remoting project settings.

diff --git a/Engine/Plugins/Runtime/WindowsMixedReality/Source/WindowsMixedRealityHMD/Private/WindowsMixedRealityHMD.cpp b/Engine/Plugins/Runtime/WindowsMixedReality/Source/WindowsMixedRealityHMD/Private/WindowsMixedRealityHMD.cpp
--- a/Engine/Plugins/Runtime/WindowsMixedReality/Source/WindowsMixedRealityHMD/Private/WindowsMixedRealityHMD.cpp
+++ b/Engine/Plugins/Runtime/WindowsMixedReality/Source/WindowsMixedRealityHMD/Private/WindowsMixedRealityHMD.cpp
@@ -106,6 +106,8 @@ namespace WindowsMixedReality
 #if WITH_WINDOWS_MIXED_REALITY
 			IHeadMountedDisplayModule::StartupModule();
 
+			GConfig->GetBool(TEXT("/Script/WindowsMixedRealityRuntimeSettings.WindowsMixedRealityRuntimeSettings"), TEXT("bEnableRemotingForEditor"), bIsRemotingEnabledForEditor, GEngineIni);
+
 			// Get the base directory of this plugin
 			FString BaseDir = IPluginManager::Get().FindPlugin("WindowsMixedReality")->GetBaseDir();
 
@@ -173,11 +175,21 @@ namespace WindowsMixedReality
 
 #if WITH_WINDOWS_MIXED_REALITY
 		MixedRealityInterop* HMD = nullptr;
+
 #endif
+
+		bool bIsRemotingEnabledForEditor = true;
 	};
 
 	TSharedPtr< class IXRTrackingSystem, ESPMode::ThreadSafe > FWindowsMixedRealityHMDPlugin::CreateTrackingSystem()
 	{
+#if WITH_EDITOR
+		if (GIsEditor && !bIsRemotingEnabledForEditor)
+		{
+			return nullptr;
+		}
+#endif
+
 #if WITH_WINDOWS_MIXED_REALITY
 		if (HMD)
 		{
diff --git a/Engine/Plugins/Runtime/WindowsMixedReality/Source/WindowsMixedRealityPlatformEditor/Private/WindowsMixedRealityDetails.cpp b/Engine/Plugins/Runtime/WindowsMixedReality/Source/WindowsMixedRealityPlatformEditor/Private/WindowsMixedRealityDetails.cpp
--- a/Engine/Plugins/Runtime/WindowsMixedReality/Source/WindowsMixedRealityPlatformEditor/Private/WindowsMixedRealityDetails.cpp
+++ b/Engine/Plugins/Runtime/WindowsMixedReality/Source/WindowsMixedRealityPlatformEditor/Private/WindowsMixedRealityDetails.cpp
@@ -28,6 +28,7 @@ void FWindowsMixedRealityDetails::CustomizeDetails(IDetailLayoutBuilder& DetailB
 			SNew(SButton)
 			.Text(LOCTEXT("Connect", "Connect"))
 			.OnClicked_Raw(this, &FWindowsMixedRealityDetails::OnConnectButtonClicked)
+			.IsEnabled_Raw(this, &FWindowsMixedRealityDetails::AreButtonsEnabled)
 		];
 	
 	remotingCategory.AddCustomRow(LOCTEXT("Disconnect Button", "Disconnect Button"))
@@ -35,6 +36,7 @@ void FWindowsMixedRealityDetails::CustomizeDetails(IDetailLayoutBuilder& DetailB
 			SNew(SButton)
 			.Text(LOCTEXT("Disconnect", "Disconnect"))
 			.OnClicked_Raw(this, &FWindowsMixedRealityDetails::OnDisconnectButtonClicked)
+			.IsEnabled_Raw(this, &FWindowsMixedRealityDetails::AreButtonsEnabled)
 		];
 }
 
@@ -60,4 +62,10 @@ FReply FWindowsMixedRealityDetails::OnDisconnectButtonClicked()
 	return FReply::Handled();
 }
 
+bool FWindowsMixedRealityDetails::AreButtonsEnabled() const
+{
+	UWindowsMixedRealityRuntimeSettings* settings = UWindowsMixedRealityRuntimeSettings::Get();
+	return settings->bEnableRemotingForEditor;
+}
+
 #undef LOCTEXT_NAMESPACE
diff --git a/Engine/Plugins/Runtime/WindowsMixedReality/Source/WindowsMixedRealityPlatformEditor/Private/WindowsMixedRealityDetails.h b/Engine/Plugins/Runtime/WindowsMixedReality/Source/WindowsMixedRealityPlatformEditor/Private/WindowsMixedRealityDetails.h
--- a/Engine/Plugins/Runtime/WindowsMixedReality/Source/WindowsMixedRealityPlatformEditor/Private/WindowsMixedRealityDetails.h
+++ b/Engine/Plugins/Runtime/WindowsMixedReality/Source/WindowsMixedRealityPlatformEditor/Private/WindowsMixedRealityDetails.h
@@ -17,4 +17,5 @@ public:
 private:
 	FReply OnConnectButtonClicked();
 	FReply OnDisconnectButtonClicked();
+	bool AreButtonsEnabled() const;
 };
\ No newline at end of file
diff --git a/Engine/Plugins/Runtime/WindowsMixedReality/Source/WindowsMixedRealityRuntimeSettings/Classes/WindowsMixedRealityRuntimeSettings.h b/Engine/Plugins/Runtime/WindowsMixedReality/Source/WindowsMixedRealityRuntimeSettings/Classes/WindowsMixedRealityRuntimeSettings.h
--- a/Engine/Plugins/Runtime/WindowsMixedReality/Source/WindowsMixedRealityRuntimeSettings/Classes/WindowsMixedRealityRuntimeSettings.h
+++ b/Engine/Plugins/Runtime/WindowsMixedReality/Source/WindowsMixedRealityRuntimeSettings/Classes/WindowsMixedRealityRuntimeSettings.h
@@ -31,10 +31,13 @@ public:
 	virtual void PostEditChangeProperty(struct FPropertyChangedEvent& PropertyChangedEvent) override;
 #endif // WITH_EDITOR
 
+	UPROPERTY(GlobalConfig, EditAnywhere, Category = "Holographic Remoting", Meta = (DisplayName = "Enable Remoting For Editor (Requires Restart)", Tooltip = "If true WMR is a valid HMD even if none is connected so that one could connect via remoting.  Editor restart required."))
+	bool bEnableRemotingForEditor = false;
+
 	/** The IP of the HoloLens to remote to. */
-	UPROPERTY(GlobalConfig, EditAnywhere, Category = "Holographic Remoting", Meta = (DisplayName = "IP of HoloLens to remote to."))
+	UPROPERTY(GlobalConfig, EditAnywhere, Category = "Holographic Remoting", Meta = (EditCondition = "bEnableRemotingForEditor", DisplayName = "IP of HoloLens to remote to."))
 	FString RemoteHoloLensIP;
 
-	UPROPERTY(GlobalConfig, EditAnywhere, Category = "Holographic Remoting", Meta = (DisplayName = "Max network transfer rate (kb/s)"))
+	UPROPERTY(GlobalConfig, EditAnywhere, Category = "Holographic Remoting", Meta = (EditCondition = "bEnableRemotingForEditor", DisplayName = "Max network transfer rate (kb/s)"))
 	unsigned int MaxBitrate = 4000;
 };
