5689730 UE-71585 VR Mode and VR Preview are able to be selected when the WMR plugin is enabled and there is no headset attached
-flipped the default to false

diff --git a/Engine/Plugins/Runtime/WindowsMixedReality/Source/WindowsMixedRealityHMD/Private/WindowsMixedRealityHMD.cpp b/Engine/Plugins/Runtime/WindowsMixedReality/Source/WindowsMixedRealityHMD/Private/WindowsMixedRealityHMD.cpp
--- a/Engine/Plugins/Runtime/WindowsMixedReality/Source/WindowsMixedRealityHMD/Private/WindowsMixedRealityHMD.cpp
+++ b/Engine/Plugins/Runtime/WindowsMixedReality/Source/WindowsMixedRealityHMD/Private/WindowsMixedRealityHMD.cpp
@@ -178,7 +178,7 @@ namespace WindowsMixedReality
 
 #endif
 
-		bool bIsRemotingEnabledForEditor = true;
+		bool bIsRemotingEnabledForEditor = false;
 	};
 
 	TSharedPtr< class IXRTrackingSystem, ESPMode::ThreadSafe > FWindowsMixedRealityHMDPlugin::CreateTrackingSystem()
