The WindowsMR tracking system failed to be initialized unless bIsRemotingEnabledForEditor was true. Fixed logic to always return an HMD interface if a headset is connected, and then return an interface in the editor only if this flag is true and there's no active headset.

diff --git a/Engine/Plugins/Runtime/WindowsMixedReality/Source/WindowsMixedRealityHMD/Private/WindowsMixedRealityHMD.cpp b/Engine/Plugins/Runtime/WindowsMixedReality/Source/WindowsMixedRealityHMD/Private/WindowsMixedRealityHMD.cpp
--- a/Engine/Plugins/Runtime/WindowsMixedReality/Source/WindowsMixedRealityHMD/Private/WindowsMixedRealityHMD.cpp
+++ b/Engine/Plugins/Runtime/WindowsMixedReality/Source/WindowsMixedRealityHMD/Private/WindowsMixedRealityHMD.cpp
@@ -183,12 +183,6 @@ namespace WindowsMixedReality
 
 	TSharedPtr< class IXRTrackingSystem, ESPMode::ThreadSafe > FWindowsMixedRealityHMDPlugin::CreateTrackingSystem()
 	{
-#if WITH_EDITOR
-		if (GIsEditor && !bIsRemotingEnabledForEditor)
-		{
-			return nullptr;
-		}
-#endif
 
 #if WITH_WINDOWS_MIXED_REALITY
 		if (HMD)
@@ -196,7 +190,20 @@ namespace WindowsMixedReality
 			auto WindowsMRHMD = FSceneViewExtensions::NewExtension<WindowsMixedReality::FWindowsMixedRealityHMD>(HMD);
 			if (WindowsMRHMD->IsInitialized())
 			{
-				return WindowsMRHMD;
+
+#if WITH_EDITOR
+				// If this is an editor build and we don't have an available HMD display, 
+				// only return a valid HMD interface if we allow remoting (since there's the
+				// potential for remoting to be set up later).
+				bool bPotentialDelayedConnect = (GIsEditor && bIsRemotingEnabledForEditor);
+#else
+				// If we're not in the editor, provide the capability for remoting to be connected.
+				bool bPotentialDelayedConnect = true;
+#endif
+				if (HMD->IsAvailable() || bPotentialDelayedConnect)
+				{
+					return WindowsMRHMD;
+				}
 			}
 		}
 #endif
