Fix CIS with proper include for GWorld access.

diff --git a/Engine/Plugins/Runtime/WindowsMixedReality/Source/WindowsMixedRealityEyeTracker/Private/WindowsMixedRealityEyeTracker.cpp b/Engine/Plugins/Runtime/WindowsMixedReality/Source/WindowsMixedRealityEyeTracker/Private/WindowsMixedRealityEyeTracker.cpp
--- a/Engine/Plugins/Runtime/WindowsMixedReality/Source/WindowsMixedRealityEyeTracker/Private/WindowsMixedRealityEyeTracker.cpp
+++ b/Engine/Plugins/Runtime/WindowsMixedReality/Source/WindowsMixedRealityEyeTracker/Private/WindowsMixedRealityEyeTracker.cpp
@@ -2,6 +2,7 @@
 #include "EyeTrackerTypes.h"
 #include "IEyeTracker.h"
 #include "Modules/ModuleManager.h"
+#include "Engine/World.h"
 
 #if WITH_WINDOWS_MIXED_REALITY
 	#include "MixedRealityInterop.h"
