Fix CIS by adding include for world settings to WMR Eye Tracking

diff --git a/Engine/Plugins/Runtime/WindowsMixedReality/Source/WindowsMixedRealityEyeTracker/Private/WindowsMixedRealityEyeTracker.cpp b/Engine/Plugins/Runtime/WindowsMixedReality/Source/WindowsMixedRealityEyeTracker/Private/WindowsMixedRealityEyeTracker.cpp
--- a/Engine/Plugins/Runtime/WindowsMixedReality/Source/WindowsMixedRealityEyeTracker/Private/WindowsMixedRealityEyeTracker.cpp
+++ b/Engine/Plugins/Runtime/WindowsMixedReality/Source/WindowsMixedRealityEyeTracker/Private/WindowsMixedRealityEyeTracker.cpp
@@ -2,6 +2,7 @@
 #include "EyeTrackerTypes.h"
 #include "IEyeTracker.h"
 #include "Modules/ModuleManager.h"
+#include "GameFramework/WorldSettings.h"
 #include "Engine/World.h"
 
 #if WITH_WINDOWS_MIXED_REALITY
