Fixed warning when building UBT on Mac

diff --git a/Engine/Plugins/Runtime/WindowsMixedReality/Source/MixedRealityInteropLibrary/MixedRealityInteropLibrary.Build.cs b/Engine/Plugins/Runtime/WindowsMixedReality/Source/MixedRealityInteropLibrary/MixedRealityInteropLibrary.Build.cs
--- a/Engine/Plugins/Runtime/WindowsMixedReality/Source/MixedRealityInteropLibrary/MixedRealityInteropLibrary.Build.cs
+++ b/Engine/Plugins/Runtime/WindowsMixedReality/Source/MixedRealityInteropLibrary/MixedRealityInteropLibrary.Build.cs
@@ -9,8 +9,6 @@ public class MixedRealityInteropLibrary : ModuleRules
 	{
 		Type = ModuleType.External;
 
-		string ThirdPartyPath = Path.GetFullPath(Path.Combine(ModuleDirectory, "../ThirdParty"));
-
 		if (Target.Platform == UnrealTargetPlatform.Win32 || Target.Platform == UnrealTargetPlatform.Win64)
 		{
 			string LibName = "MixedRealityInterop";
@@ -26,7 +24,7 @@ public class MixedRealityInteropLibrary : ModuleRules
 			PublicDelayLoadDLLs.Add(DLLName);
 			RuntimeDependencies.Add(PluginDirectory + "/Binaries/ThirdParty/MixedRealityInteropLibrary/" + Target.Platform.ToString() + "/" + DLLName);
 			
-			//Hologram remoting dlls
+			// Hologram remoting dlls
 			if (Target.Platform == UnrealTargetPlatform.Win64 && Target.bBuildEditor == true)
 			{
 				string[] Dlls = { "HolographicAppRemoting.dll" , "PerceptionDevice.dll" };
@@ -43,4 +41,4 @@ public class MixedRealityInteropLibrary : ModuleRules
 			PublicAdditionalLibraries.Add("MixedRealityInteropHoloLens.lib");
 		}
 	}
-}
\ No newline at end of file
+}
