Wrapped the HoloLens eye gaze code in #if PLATFORM_HOLOLENS

diff --git a/Engine/Plugins/Runtime/WindowsMixedReality/Source/WindowsMixedRealityEyeTracker/Private/WindowsMixedRealityEyeTracker.cpp b/Engine/Plugins/Runtime/WindowsMixedReality/Source/WindowsMixedRealityEyeTracker/Private/WindowsMixedRealityEyeTracker.cpp
--- a/Engine/Plugins/Runtime/WindowsMixedReality/Source/WindowsMixedRealityEyeTracker/Private/WindowsMixedRealityEyeTracker.cpp
+++ b/Engine/Plugins/Runtime/WindowsMixedReality/Source/WindowsMixedRealityEyeTracker/Private/WindowsMixedRealityEyeTracker.cpp
@@ -5,13 +5,13 @@
 #include "GameFramework/WorldSettings.h"
 #include "Engine/World.h"
 
-#if WITH_WINDOWS_MIXED_REALITY
+#if WITH_WINDOWS_MIXED_REALITY && PLATFORM_HOLOLENS
 	#include "MixedRealityInterop.h"
 #endif
 
+#if WITH_WINDOWS_MIXED_REALITY && PLATFORM_HOLOLENS
 static WindowsMixedReality::MixedRealityInterop hmd;
 
-
 class FWindowsMixedRealityEyeTracker :
 	public IEyeTracker
 {
@@ -70,6 +70,7 @@ private:
 		return GWorld ? GWorld->GetWorldSettings()->WorldToMeters : 100.f;
 	}
 };
+#endif
 
 class FWindowsMixedRealityEyeTrackerModule :
 	public IEyeTrackerModule
@@ -92,12 +93,20 @@ public:
 
 	virtual bool IsEyeTrackerConnected() const override
 	{
+#if WITH_WINDOWS_MIXED_REALITY && PLATFORM_HOLOLENS
 		return hmd.SupportsEyeTracking();
+#else
+		return false;
+#endif
 	}
 	
 	virtual TSharedPtr< class IEyeTracker, ESPMode::ThreadSafe > CreateEyeTracker() override
 	{
+#if WITH_WINDOWS_MIXED_REALITY && PLATFORM_HOLOLENS
 		return TSharedPtr< class IEyeTracker, ESPMode::ThreadSafe >(new FWindowsMixedRealityEyeTracker);
+#else
+		return TSharedPtr< class IEyeTracker, ESPMode::ThreadSafe >();
+#endif
 	}
 };
 
