LiveLink Controller and Hand Tracking merge fixes.

diff --git a/Engine/Plugins/Animation/LiveLink/Source/LiveLink/Private/LiveLinkMotionController.h b/Engine/Plugins/Animation/LiveLink/Source/LiveLink/Private/LiveLinkMotionController.h
--- a/Engine/Plugins/Animation/LiveLink/Source/LiveLink/Private/LiveLinkMotionController.h
+++ b/Engine/Plugins/Animation/LiveLink/Source/LiveLink/Private/LiveLinkMotionController.h
@@ -174,6 +174,8 @@ public:
 		}
 	}
 
+	virtual bool GetHandJointPosition(const FName MotionSource, int jointIndex, FVector& OutPosition) const override { return false; }
+
 private:
 	FLiveLinkSubjectKey GetSubjectKeyFromMotionSource(FName MotionSource) const
 	{
diff --git a/Engine/Plugins/Runtime/WindowsMixedReality/Source/WindowsMixedRealityHandTracking/Private/WindowsMixedRealityHandTracking.cpp b/Engine/Plugins/Runtime/WindowsMixedReality/Source/WindowsMixedRealityHandTracking/Private/WindowsMixedRealityHandTracking.cpp
--- a/Engine/Plugins/Runtime/WindowsMixedReality/Source/WindowsMixedRealityHandTracking/Private/WindowsMixedRealityHandTracking.cpp
+++ b/Engine/Plugins/Runtime/WindowsMixedReality/Source/WindowsMixedRealityHandTracking/Private/WindowsMixedRealityHandTracking.cpp
@@ -117,14 +117,15 @@ private:
 IMPLEMENT_MODULE(FWindowsMixedRealityHandTrackingModule, WindowsMixedRealityHandTracking);
 
 
-FName FWindowsMixedRealityHandTracking::LiveLinkLeftHandTrackingSubjectName(TEXT("WMRLeftHand"));
-FName FWindowsMixedRealityHandTracking::LiveLinkRightHandTrackingSubjectName(TEXT("WMRRightHand"));
+FLiveLinkSubjectName FWindowsMixedRealityHandTracking::LiveLinkLeftHandTrackingSubjectName(TEXT("WMRLeftHand"));
+FLiveLinkSubjectName FWindowsMixedRealityHandTracking::LiveLinkRightHandTrackingSubjectName(TEXT("WMRRightHand"));
 
 
 FWindowsMixedRealityHandTracking::FWindowsMixedRealityHandTracking(const TSharedRef<FGenericApplicationMessageHandler>& InMessageHandler)
 	: MessageHandler(InMessageHandler)
 	, DeviceIndex(0)
 	, bIsHandTrackingStateValid(true)
+	, LiveLinkSkeletonStaticData(FLiveLinkSkeletonStaticData::StaticStruct())
 {
 	// Register "MotionController" modular feature manually
 	IModularFeatures::Get().RegisterModularFeature(GetModularFeatureName(), this);
diff --git a/Engine/Plugins/Runtime/WindowsMixedReality/Source/WindowsMixedRealityHandTracking/Private/WindowsMixedRealityHandTracking.h b/Engine/Plugins/Runtime/WindowsMixedReality/Source/WindowsMixedRealityHandTracking/Private/WindowsMixedRealityHandTracking.h
--- a/Engine/Plugins/Runtime/WindowsMixedReality/Source/WindowsMixedRealityHandTracking/Private/WindowsMixedRealityHandTracking.h
+++ b/Engine/Plugins/Runtime/WindowsMixedReality/Source/WindowsMixedRealityHandTracking/Private/WindowsMixedRealityHandTracking.h
@@ -13,6 +13,7 @@
 #include "InputCoreTypes.h"
 #include "ILiveLinkSource.h"
 #include "IInputDevice.h"
+#include "Roles/LiveLinkAnimationTypes.h"
 
 /**
   * WindowsMixedReality HandTracking
@@ -46,7 +47,7 @@ public:
 
 	// ILiveLinkSource interface
 	virtual void ReceiveClient(ILiveLinkClient* InClient, FGuid InSourceGuid) override;
-	virtual bool IsSourceStillValid() override;
+	virtual bool IsSourceStillValid() const override;
 	virtual bool RequestSourceShutdown() override;
 	virtual FText GetSourceMachineName() const override;
 	virtual FText GetSourceStatus() const override;
@@ -103,10 +104,12 @@ private:
 	/** Our identifier in LiveLink */
 	FGuid LiveLinkSourceGuid;
 
-	static FName LiveLinkLeftHandTrackingSubjectName;
-	static FName LiveLinkRightHandTrackingSubjectName;
+	static FLiveLinkSubjectName LiveLinkLeftHandTrackingSubjectName;
+	static FLiveLinkSubjectName LiveLinkRightHandTrackingSubjectName;
+	FLiveLinkSubjectKey LiveLinkLeftHandTrackingSubjectKey;
+	FLiveLinkSubjectKey LiveLinkRightHandTrackingSubjectKey;
 	bool bNewLiveLinkClient = false;
-	FLiveLinkRefSkeleton LiveLinkRefSkeleton;
+	FLiveLinkStaticDataStruct LiveLinkSkeletonStaticData;
 };
 
 DEFINE_LOG_CATEGORY_STATIC(LogWindowsMixedRealityHandTracking, Display, All);
diff --git a/Engine/Plugins/Runtime/WindowsMixedReality/Source/WindowsMixedRealityHandTracking/Private/WindowsMixedRealityHandTrackingLiveLink.cpp b/Engine/Plugins/Runtime/WindowsMixedReality/Source/WindowsMixedRealityHandTracking/Private/WindowsMixedRealityHandTrackingLiveLink.cpp
--- a/Engine/Plugins/Runtime/WindowsMixedReality/Source/WindowsMixedRealityHandTracking/Private/WindowsMixedRealityHandTrackingLiveLink.cpp
+++ b/Engine/Plugins/Runtime/WindowsMixedReality/Source/WindowsMixedRealityHandTracking/Private/WindowsMixedRealityHandTrackingLiveLink.cpp
@@ -10,6 +10,7 @@
 #include "Engine/Engine.h"
 #include "LiveLinkSourceFactory.h"
 #include "ILiveLinkClient.h"
+#include "Roles/LiveLinkAnimationRole.h"
 
 #define LOCTEXT_NAMESPACE "WindowsMixedRealityHandTracking"
 
@@ -19,10 +20,15 @@ void FWindowsMixedRealityHandTracking::ReceiveClient(ILiveLinkClient* InClient,
 	LiveLinkSourceGuid = InSourceGuid;
 	bNewLiveLinkClient = true;
 
+	LiveLinkLeftHandTrackingSubjectKey.Source = InSourceGuid;
+	LiveLinkLeftHandTrackingSubjectKey.SubjectName = LiveLinkLeftHandTrackingSubjectName;
+	LiveLinkRightHandTrackingSubjectKey.Source = InSourceGuid;
+	LiveLinkRightHandTrackingSubjectKey.SubjectName = LiveLinkRightHandTrackingSubjectName;
+
 	UpdateLiveLink();
 }
 
-bool FWindowsMixedRealityHandTracking::IsSourceStillValid()
+bool FWindowsMixedRealityHandTracking::IsSourceStillValid() const
 {
 	return LiveLinkClient != nullptr;
 }
@@ -101,8 +107,10 @@ void FWindowsMixedRealityHandTracking::SetupLiveLinkData()
 	BoneParents.Add(23);	// LittleDistal -> LittleIntermediate
 	BoneParents.Add(24);	// LittleTip -> LittleDistal
 
-	LiveLinkRefSkeleton.SetBoneNames(BoneNames);
-	LiveLinkRefSkeleton.SetBoneParents(BoneParents);
+	FLiveLinkSkeletonStaticData* SkeletonDataPtr = LiveLinkSkeletonStaticData.Cast<FLiveLinkSkeletonStaticData>();
+	check(SkeletonDataPtr);
+	SkeletonDataPtr->SetBoneNames(BoneNames);
+	SkeletonDataPtr->SetBoneParents(BoneParents);
 }
 
 void FWindowsMixedRealityHandTracking::UpdateLiveLinkTransforms(TArray<FTransform>& OutTransforms, const FWindowsMixedRealityHandTracking::FHandState& HandState)
@@ -132,47 +140,56 @@ void FWindowsMixedRealityHandTracking::UpdateLiveLink()
 
 	if (LiveLinkClient)
 	{
-		static FLiveLinkFrameData LiveLinkLeftFrame;
-		static FLiveLinkFrameData LiveLinkRightFrame;
+		FLiveLinkFrameDataStruct LiveLinkLeftFrameDataStruct(FLiveLinkAnimationFrameData::StaticStruct());
+		FLiveLinkAnimationFrameData& LiveLinkLeftAnimationFrameData = *LiveLinkLeftFrameDataStruct.Cast<FLiveLinkAnimationFrameData>();
+		FLiveLinkFrameDataStruct LiveLinkRightFrameDataStruct(FLiveLinkAnimationFrameData::StaticStruct());
+		FLiveLinkAnimationFrameData& LiveLinkRightAnimationFrameData = *LiveLinkRightFrameDataStruct.Cast<FLiveLinkAnimationFrameData>();
+
+
 		static bool bInitialized = false;
 		if (!bInitialized)
 		{
 			SetupLiveLinkData();
 
-			LiveLinkLeftFrame.Transforms.Reserve(EWMRHandKeypointCount);
-			LiveLinkRightFrame.Transforms.Reserve(EWMRHandKeypointCount);
+			LiveLinkLeftAnimationFrameData.Transforms.Reserve(EWMRHandKeypointCount);
+			LiveLinkRightAnimationFrameData.Transforms.Reserve(EWMRHandKeypointCount);
 			// Init to identity all of the Keypoint transforms
 			for (uint32 Count = 0; Count < EWMRHandKeypointCount; ++Count)
 			{
-				LiveLinkLeftFrame.Transforms.Add(FTransform::Identity);
-				LiveLinkRightFrame.Transforms.Add(FTransform::Identity);
+				LiveLinkLeftAnimationFrameData.Transforms.Add(FTransform::Identity);
+				LiveLinkRightAnimationFrameData.Transforms.Add(FTransform::Identity);
 			}
 
 			// Set some metadata to specify handedness
-			LiveLinkLeftFrame.MetaData.StringMetaData.Add(FName(TEXT("Handedness")), TEXT("Left"));
-			LiveLinkRightFrame.MetaData.StringMetaData.Add(FName(TEXT("Handedness")), TEXT("Right"));
+			LiveLinkLeftAnimationFrameData.MetaData.StringMetaData.Add(FName(TEXT("Handedness")), TEXT("Left"));
+			LiveLinkRightAnimationFrameData.MetaData.StringMetaData.Add(FName(TEXT("Handedness")), TEXT("Right"));
 
 			bInitialized = true;
 		}
 
 		if (bNewLiveLinkClient)
 		{
-			LiveLinkClient->ClearSubject(LiveLinkLeftHandTrackingSubjectName);
-			LiveLinkClient->ClearSubject(LiveLinkRightHandTrackingSubjectName);
-			LiveLinkClient->PushSubjectSkeleton(LiveLinkSourceGuid, LiveLinkLeftHandTrackingSubjectName, LiveLinkRefSkeleton);
-			LiveLinkClient->PushSubjectSkeleton(LiveLinkSourceGuid, LiveLinkRightHandTrackingSubjectName, LiveLinkRefSkeleton);
+			FLiveLinkStaticDataStruct SkeletalDataLeft;
+			SkeletalDataLeft.InitializeWith(LiveLinkSkeletonStaticData);
+			FLiveLinkStaticDataStruct SkeletalDataRight;
+			SkeletalDataRight.InitializeWith(LiveLinkSkeletonStaticData);
+
+			LiveLinkClient->RemoveSubject_AnyThread(LiveLinkLeftHandTrackingSubjectKey);
+			LiveLinkClient->RemoveSubject_AnyThread(LiveLinkRightHandTrackingSubjectKey);
+			LiveLinkClient->PushSubjectStaticData_AnyThread(LiveLinkLeftHandTrackingSubjectKey, ULiveLinkAnimationRole::StaticClass(), MoveTemp(SkeletalDataLeft));
+			LiveLinkClient->PushSubjectStaticData_AnyThread(LiveLinkRightHandTrackingSubjectKey, ULiveLinkAnimationRole::StaticClass(), MoveTemp(SkeletalDataRight));
 			bNewLiveLinkClient = false;
 		}
 
-		LiveLinkLeftFrame.WorldTime = LiveLinkRightFrame.WorldTime = FPlatformTime::Seconds();
+		LiveLinkLeftAnimationFrameData.WorldTime = LiveLinkRightAnimationFrameData.WorldTime = FPlatformTime::Seconds();
 
 		// Update the transforms for each subject from tracking data
-		UpdateLiveLinkTransforms(LiveLinkLeftFrame.Transforms, GetLeftHandState());
-		UpdateLiveLinkTransforms(LiveLinkRightFrame.Transforms, GetRightHandState());
+		UpdateLiveLinkTransforms(LiveLinkLeftAnimationFrameData.Transforms, GetLeftHandState());
+		UpdateLiveLinkTransforms(LiveLinkRightAnimationFrameData.Transforms, GetRightHandState());
 
 		// Share the data locally with the LiveLink client
-		LiveLinkClient->PushSubjectData(LiveLinkSourceGuid, LiveLinkLeftHandTrackingSubjectName, LiveLinkLeftFrame);
-		LiveLinkClient->PushSubjectData(LiveLinkSourceGuid, LiveLinkRightHandTrackingSubjectName, LiveLinkRightFrame);
+		LiveLinkClient->PushSubjectFrameData_AnyThread(LiveLinkLeftHandTrackingSubjectKey, MoveTemp(LiveLinkLeftFrameDataStruct));
+		LiveLinkClient->PushSubjectFrameData_AnyThread(LiveLinkRightHandTrackingSubjectKey, MoveTemp(LiveLinkRightFrameDataStruct));
 	}
 }
 
diff --git a/Engine/Plugins/Runtime/WindowsMixedReality/Source/WindowsMixedRealityHandTracking/Private/WindowsMixedRealityHandTrackingLiveLinkRemapAsset.cpp b/Engine/Plugins/Runtime/WindowsMixedReality/Source/WindowsMixedRealityHandTracking/Private/WindowsMixedRealityHandTrackingLiveLinkRemapAsset.cpp
--- a/Engine/Plugins/Runtime/WindowsMixedReality/Source/WindowsMixedRealityHandTracking/Private/WindowsMixedRealityHandTrackingLiveLinkRemapAsset.cpp
+++ b/Engine/Plugins/Runtime/WindowsMixedReality/Source/WindowsMixedRealityHandTracking/Private/WindowsMixedRealityHandTrackingLiveLinkRemapAsset.cpp
@@ -33,10 +33,13 @@ void UWindowsMixedRealityHandTrackingLiveLinkRemapAsset::PostInitProperties()
 }
 #endif
 
-void UWindowsMixedRealityHandTrackingLiveLinkRemapAsset::BuildPoseForSubject(float DeltaTime, const FLiveLinkSubjectFrame& InFrame, FCompactPose& OutPose, FBlendedCurve& OutCurve)
+void UWindowsMixedRealityHandTrackingLiveLinkRemapAsset::BuildPoseForSubject(float DeltaTime, const FLiveLinkSkeletonStaticData* InSkeletonData, const FLiveLinkAnimationFrameData* InFrameData, FCompactPose& OutPose, FBlendedCurve& OutCurve)
 {
+	check(InSkeletonData);
+	check(InFrameData);
+
 	// Transform Bone Names
-	const TArray<FName>& SourceBoneNames = InFrame.RefSkeleton.GetBoneNames();
+	const TArray<FName>& SourceBoneNames = InSkeletonData->GetBoneNames();
 
 	TArray<FName, TMemStackAllocator<>> TransformedBoneNames;
 	TransformedBoneNames.Reserve(SourceBoneNames.Num());
@@ -50,7 +53,7 @@ void UWindowsMixedRealityHandTrackingLiveLinkRemapAsset::BuildPoseForSubject(flo
 	{
 		FName BoneName = TransformedBoneNames[i];
 
-		FTransform BoneTransform = GetRetargetedTransform(InFrame, i);
+		FTransform BoneTransform = GetRetargetedTransform(InFrameData, i);
 
 		int32 MeshIndex = OutPose.GetBoneContainer().GetPoseBoneIndexForBoneName(BoneName);
 		if (MeshIndex != INDEX_NONE)
@@ -68,7 +71,7 @@ void UWindowsMixedRealityHandTrackingLiveLinkRemapAsset::BuildPoseForSubject(flo
 				NewRot.W = Sign(SwizzleW) * OldRotVector[StaticCast<uint8>(SwizzleW) % 4];
 
 				// Left hand bones need to be flipped
-				const FString* Handedness = InFrame.MetaData.StringMetaData.Find(TEXT("Handedness"));
+				const FString* Handedness = InFrameData->MetaData.StringMetaData.Find(TEXT("Handedness"));
 				if (Handedness && Handedness->Equals(TEXT("Left")))
 				{
 					NewRot.X = -NewRot.X;
@@ -104,9 +107,11 @@ FName UWindowsMixedRealityHandTrackingLiveLinkRemapAsset::GetRemappedBoneName(FN
 	return BoneName;
 }
 
-FTransform UWindowsMixedRealityHandTrackingLiveLinkRemapAsset::GetRetargetedTransform(const FLiveLinkSubjectFrame& InFrame, int TransformIndex) const
+FTransform UWindowsMixedRealityHandTrackingLiveLinkRemapAsset::GetRetargetedTransform(const FLiveLinkAnimationFrameData* InFrameData, int TransformIndex) const
 {
-	FTransform OutTransform = InFrame.Transforms[TransformIndex];
+	check(InFrameData);
+
+	FTransform OutTransform = InFrameData->Transforms[TransformIndex];
 
 	if (!bHasMetacarpals && (
 		TransformIndex == static_cast<int>(EWMRHandKeypoint::ThumbProximal)		|| 
@@ -116,7 +121,7 @@ FTransform UWindowsMixedRealityHandTrackingLiveLinkRemapAsset::GetRetargetedTran
 		TransformIndex == static_cast<int>(EWMRHandKeypoint::LittleProximal)	))
 	{
 		// Metacarpal is always one entry before the Proximal
-		OutTransform = InFrame.Transforms[TransformIndex - 1] * OutTransform;
+		OutTransform = InFrameData->Transforms[TransformIndex - 1] * OutTransform;
 	}
 
 	return OutTransform;
diff --git a/Engine/Plugins/Runtime/WindowsMixedReality/Source/WindowsMixedRealityHandTracking/Private/WindowsMixedRealityHandTrackingLiveLinkRemapAsset.h b/Engine/Plugins/Runtime/WindowsMixedReality/Source/WindowsMixedRealityHandTracking/Private/WindowsMixedRealityHandTrackingLiveLinkRemapAsset.h
--- a/Engine/Plugins/Runtime/WindowsMixedReality/Source/WindowsMixedRealityHandTracking/Private/WindowsMixedRealityHandTrackingLiveLinkRemapAsset.h
+++ b/Engine/Plugins/Runtime/WindowsMixedReality/Source/WindowsMixedRealityHandTracking/Private/WindowsMixedRealityHandTrackingLiveLinkRemapAsset.h
@@ -34,7 +34,7 @@ class UWindowsMixedRealityHandTrackingLiveLinkRemapAsset :
 	GENERATED_UCLASS_BODY()
 
 public:
-	virtual void BuildPoseForSubject(float DeltaTime, const FLiveLinkSubjectFrame& InFrame, FCompactPose& OutPose, FBlendedCurve& OutCurve) override;
+	virtual void BuildPoseForSubject(float DeltaTime, const FLiveLinkSkeletonStaticData* InSkeletonData, const FLiveLinkAnimationFrameData* InFrameData, FCompactPose& OutPose, FBlendedCurve& OutCurve) override;
 
 	/** If true, remap the full human hand skeleton including metacarpals */
 	UPROPERTY(EditAnywhere, Category = "LiveLink")
@@ -64,5 +64,5 @@ public:
 private:
 	FName GetRemappedBoneName(FName BoneName) const;
 
-	FTransform GetRetargetedTransform(const FLiveLinkSubjectFrame& InFrame, int TransformIndex) const;
+	FTransform GetRetargetedTransform(const FLiveLinkAnimationFrameData* InFrameData, int TransformIndex) const;
 };
