LiveLink HandTracking runtime fix for transform array having no entries.

diff --git a/Engine/Plugins/Runtime/WindowsMixedReality/Source/WindowsMixedRealityHandTracking/Private/WindowsMixedRealityHandTrackingLiveLink.cpp b/Engine/Plugins/Runtime/WindowsMixedReality/Source/WindowsMixedRealityHandTracking/Private/WindowsMixedRealityHandTrackingLiveLink.cpp
--- a/Engine/Plugins/Runtime/WindowsMixedReality/Source/WindowsMixedRealityHandTracking/Private/WindowsMixedRealityHandTrackingLiveLink.cpp
+++ b/Engine/Plugins/Runtime/WindowsMixedReality/Source/WindowsMixedRealityHandTracking/Private/WindowsMixedRealityHandTrackingLiveLink.cpp
@@ -143,10 +143,10 @@ void FWindowsMixedRealityHandTracking::UpdateLiveLink()
 
 	if (LiveLinkClient)
 	{
-		FLiveLinkFrameDataStruct LiveLinkLeftFrameDataStruct(FLiveLinkAnimationFrameData::StaticStruct());
-		FLiveLinkAnimationFrameData& LiveLinkLeftAnimationFrameData = *LiveLinkLeftFrameDataStruct.Cast<FLiveLinkAnimationFrameData>();
-		FLiveLinkFrameDataStruct LiveLinkRightFrameDataStruct(FLiveLinkAnimationFrameData::StaticStruct());
-		FLiveLinkAnimationFrameData& LiveLinkRightAnimationFrameData = *LiveLinkRightFrameDataStruct.Cast<FLiveLinkAnimationFrameData>();
+		static FLiveLinkFrameDataStruct LiveLinkLeftFrameDataStruct(FLiveLinkAnimationFrameData::StaticStruct());
+		static FLiveLinkAnimationFrameData& LiveLinkLeftAnimationFrameData = *LiveLinkLeftFrameDataStruct.Cast<FLiveLinkAnimationFrameData>();
+		static FLiveLinkFrameDataStruct LiveLinkRightFrameDataStruct(FLiveLinkAnimationFrameData::StaticStruct());
+		static FLiveLinkAnimationFrameData& LiveLinkRightAnimationFrameData = *LiveLinkRightFrameDataStruct.Cast<FLiveLinkAnimationFrameData>();
 
 
 		static bool bInitialized = false;
