Disable precompiling WindowsMixedRealityHMD for Win32 due to compile errors.

#rnx

diff --git a/Engine/Plugins/Runtime/WindowsMixedReality/Source/WindowsMixedRealityHMD/WindowsMixedRealityHMD.Build.cs b/Engine/Plugins/Runtime/WindowsMixedReality/Source/WindowsMixedRealityHMD/WindowsMixedRealityHMD.Build.cs
--- a/Engine/Plugins/Runtime/WindowsMixedReality/Source/WindowsMixedRealityHMD/WindowsMixedRealityHMD.Build.cs
+++ b/Engine/Plugins/Runtime/WindowsMixedReality/Source/WindowsMixedRealityHMD/WindowsMixedRealityHMD.Build.cs
@@ -130,6 +130,10 @@ namespace UnrealBuildTool.Rules
 				PrivatePCHHeaderFile = "Private/WindowsMixedRealityPrecompiled.h";
 			}
 
+			if(Target.Platform == UnrealTargetPlatform.Win32)
+            {
+				PrecompileForTargets = PrecompileTargetsType.None;
+            }
 			PublicDelayLoadDLLs.Add("QRCodesTrackerPlugin.dll");
 			RuntimeDependencies.Add(Path.Combine("$(EngineDir)/Binaries/ThirdParty/Windows/HoloLens", Target.WindowsPlatform.GetArchitectureSubpath(), "QRCodesTrackerPlugin.dll"));
 
