UE-79885 HL2 CL integrations
Merging //UE4/Partner-Microsoft-XR to Release-4.23 (//UE4/Release-[REDACTED_PHONE]
8322424 - Plumb WMR controller tracking data to IsTracked BP function - this is the two .cpp files changes
Merging //UE4/Partner-Microsoft-XR to Release-4.23 (//UE4/Release-[REDACTED_PHONE]
8513053 - Update HL2 remoting to 2.0.8.0 - this changes some of the source binaries for remoting.

diff --git a/Engine/Plugins/Runtime/WindowsMixedReality/Source/WindowsMixedRealityHandTracking/Private/WindowsMixedRealityHandTracking.cpp b/Engine/Plugins/Runtime/WindowsMixedReality/Source/WindowsMixedRealityHandTracking/Private/WindowsMixedRealityHandTracking.cpp
--- a/Engine/Plugins/Runtime/WindowsMixedReality/Source/WindowsMixedRealityHandTracking/Private/WindowsMixedRealityHandTracking.cpp
+++ b/Engine/Plugins/Runtime/WindowsMixedReality/Source/WindowsMixedRealityHandTracking/Private/WindowsMixedRealityHandTracking.cpp
@@ -242,7 +242,9 @@ bool FWindowsMixedRealityHandTracking::GetControllerOrientationAndPosition(const
 	}
 	
 	// Then call super to handle a few of the default labels, for backward compatibility
-	return FXRMotionControllerBase::GetControllerOrientationAndPosition(ControllerIndex, MotionSource, OutOrientation, OutPosition, WorldToMetersScale);
+	FXRMotionControllerBase::GetControllerOrientationAndPosition(ControllerIndex, MotionSource, OutOrientation, OutPosition, WorldToMetersScale);
+
+	return false;
 }
 
 bool FWindowsMixedRealityHandTracking::GetControllerOrientationAndPosition(const int32 ControllerIndex, const EControllerHand DeviceHand, FRotator& OutOrientation, FVector& OutPosition, float WorldToMetersScale) const
