UE-80182 HoloLens(1&2) Failure to connect via holographic remoting
-Added some log lines to tell us if we have not enabled remoting or if our build does not support wmr or hololens.

diff --git a/Engine/Plugins/Runtime/WindowsMixedReality/Source/WindowsMixedRealityHMD/Private/WindowsMixedRealityHMD.cpp b/Engine/Plugins/Runtime/WindowsMixedReality/Source/WindowsMixedRealityHMD/Private/WindowsMixedRealityHMD.cpp
--- a/Engine/Plugins/Runtime/WindowsMixedReality/Source/WindowsMixedRealityHMD/Private/WindowsMixedRealityHMD.cpp
+++ b/Engine/Plugins/Runtime/WindowsMixedReality/Source/WindowsMixedRealityHMD/Private/WindowsMixedRealityHMD.cpp
@@ -1782,10 +1782,14 @@ namespace WindowsMixedReality
 #if WITH_EDITOR
 		D3D11Device = InternalGetD3D11Device();
 
-#if WITH_WINDOWS_MIXED_REALITY
+#  if WITH_WINDOWS_MIXED_REALITY
 		HMD->SetLogCallback(WindowsMixedRealityHMD::LogForInterop);
 		HMD->ConnectToRemoteHoloLens(D3D11Device.GetReference(), ip, bitrate, isHoloLens1);
-#endif
+#  else
+		UE_LOG(LogWmrHmd, Log, TEXT("FWindowsMixedRealityHMD::ConnectToRemoteHoloLens() is doing nothing because !WITH_WINDOWS_MIXED_REALITY."));
+#  endif
+#else
+		UE_LOG(LogWmrHmd, Log, TEXT("FWindowsMixedRealityHMD::ConnectToRemoteHoloLens() is doing nothing because !WITH_EDITOR."));
 #endif
 	}
 
diff --git a/Engine/Plugins/Runtime/WindowsMixedReality/Source/WindowsMixedRealityHMD/Private/WindowsMixedRealityStatics.cpp b/Engine/Plugins/Runtime/WindowsMixedReality/Source/WindowsMixedRealityHMD/Private/WindowsMixedRealityStatics.cpp
--- a/Engine/Plugins/Runtime/WindowsMixedReality/Source/WindowsMixedRealityHMD/Private/WindowsMixedRealityStatics.cpp
+++ b/Engine/Plugins/Runtime/WindowsMixedReality/Source/WindowsMixedRealityHMD/Private/WindowsMixedRealityStatics.cpp
@@ -149,6 +149,12 @@ namespace WindowsMixedReality
 		{
 			hmd->ConnectToRemoteHoloLens(*remoteIP, bitrate, isHoloLens1);
 		}
+		else
+		{
+			UE_LOG(LogWmrHmd, Log, TEXT("FWindowsMixedRealityStatics::ConnectToRemoteHoloLens() is doing nothing because FWindowsMixedRealityHMD fetch failed.  Perhaps you need to enable remoting and restart the editor?"));
+		}
+#else
+		UE_LOG(LogWmrHmd, Log, TEXT("FWindowsMixedRealityStatics::ConnectToRemoteHoloLens() is doing nothing because PLATFORM_HOLOLENS. (You don't 'remote' when running on device or emulated device.)"));
 #endif
 	}
 
diff --git a/Engine/Plugins/Runtime/WindowsMixedReality/Source/WindowsMixedRealityPlatformEditor/Private/WindowsMixedRealityDetails.cpp b/Engine/Plugins/Runtime/WindowsMixedReality/Source/WindowsMixedRealityPlatformEditor/Private/WindowsMixedRealityDetails.cpp
--- a/Engine/Plugins/Runtime/WindowsMixedReality/Source/WindowsMixedRealityPlatformEditor/Private/WindowsMixedRealityDetails.cpp
+++ b/Engine/Plugins/Runtime/WindowsMixedReality/Source/WindowsMixedRealityPlatformEditor/Private/WindowsMixedRealityDetails.cpp
@@ -45,7 +45,8 @@ FReply FWindowsMixedRealityDetails::OnConnectButtonClicked()
 	UWindowsMixedRealityRuntimeSettings* settings = UWindowsMixedRealityRuntimeSettings::Get();
 
 	FString ip = settings->RemoteHoloLensIP;
-	UE_LOG(LogTemp, Log, TEXT("Connecting to: %s"), *ip);
+	int HoloLensType = settings->IsHoloLens1Remoting ? 1 : 2;
+	UE_LOG(LogTemp, Log, TEXT("Editor connecting to remote HoloLens%i: %s"), HoloLensType, *ip);
 
 	unsigned int bitrate = settings->MaxBitrate;
 
@@ -56,7 +57,7 @@ FReply FWindowsMixedRealityDetails::OnConnectButtonClicked()
 
 FReply FWindowsMixedRealityDetails::OnDisconnectButtonClicked()
 {
-	UE_LOG(LogTemp, Log, TEXT("Disconnecting from remote HoloLens"));
+	UE_LOG(LogTemp, Log, TEXT("Editor disconnecting from remote HoloLens"));
 	WindowsMixedReality::FWindowsMixedRealityStatics::DisconnectFromRemoteHoloLens();
 
 	return FReply::Handled();
