UEVR-1507 Hololens 2 - uncomment SetFocusPointForFrame blueprint function
-This feature was enabled in 4.23.1, but the blueprint api could not be exposed due to hotfix compatibility rules.
-The blueprint can go in for 4.24, the cvar hack for 4.23 will now warn of deprication and is scheduled for removal in 4.25.
-Updated test map to use the blueprint.

diff --git a/Engine/Plugins/Runtime/WindowsMixedReality/Source/WindowsMixedRealityHMD/Classes/WindowsMixedRealityFunctionLibrary.h b/Engine/Plugins/Runtime/WindowsMixedReality/Source/WindowsMixedRealityHMD/Classes/WindowsMixedRealityFunctionLibrary.h
--- a/Engine/Plugins/Runtime/WindowsMixedReality/Source/WindowsMixedRealityHMD/Classes/WindowsMixedRealityFunctionLibrary.h
+++ b/Engine/Plugins/Runtime/WindowsMixedReality/Source/WindowsMixedRealityHMD/Classes/WindowsMixedRealityFunctionLibrary.h
@@ -99,4 +99,12 @@ public:
 	 */
 	UFUNCTION(BlueprintCallable, Category = "WindowsMixedRealityHMD")
 	static FPointerPoseInfo GetPointerPoseInfo(EControllerHand hand);
+
+
+	/**
+	 *  Set the focus point for the current frame to stabilize your holograms.
+	 *  When run on device, the depth buffer with be used.  Use this for remoting.
+	 */
+	UFUNCTION(BlueprintCallable, Category = "WindowsMixedRealityHMD")
+	static void SetFocusPointForFrame(FVector position);
 };
diff --git a/Engine/Plugins/Runtime/WindowsMixedReality/Source/WindowsMixedRealityHMD/Private/WindowsMixedRealityFunctionLibrary.cpp b/Engine/Plugins/Runtime/WindowsMixedReality/Source/WindowsMixedRealityHMD/Private/WindowsMixedRealityFunctionLibrary.cpp
--- a/Engine/Plugins/Runtime/WindowsMixedReality/Source/WindowsMixedRealityHMD/Private/WindowsMixedRealityFunctionLibrary.cpp
+++ b/Engine/Plugins/Runtime/WindowsMixedReality/Source/WindowsMixedRealityHMD/Private/WindowsMixedRealityFunctionLibrary.cpp
@@ -116,32 +116,18 @@ FPointerPoseInfo UWindowsMixedRealityFunctionLibrary::GetPointerPoseInfo(EContro
 	return info;
 }
 
-	
-// Commented out stabilization plane blueprint implementation.  
-// This change violates the rules for binary compatibility in engine hotfixes.
-//
-// Uncomment this block and move it to the bottom of WindowsMixedRealityFunctionLibrary.h
-//	/**
-//	 *  Set the focus point for the current frame to stabilize your holograms.
-//	 *  When run on device, the depth buffer with be used.  Use this for remoting.
-//	 */
-//	UFUNCTION(BlueprintCallable, Category = "WindowsMixedRealityHMD")
-//	static void SetFocusPointForFrame(FVector position);
-
-
-// Uncomment this block and leave it here.
-//void UWindowsMixedRealityFunctionLibrary::SetFocusPointForFrame(FVector position)
-//{
-//#if WITH_WINDOWS_MIXED_REALITY
-//	WindowsMixedReality::FWindowsMixedRealityHMD* hmd = GetWindowsMixedRealityHMD();
-//	if (hmd == nullptr)
-//	{
-//		return;
-//	}
-//
-//	hmd->SetFocusPointForFrame(position);
-//#endif
-//}
+void UWindowsMixedRealityFunctionLibrary::SetFocusPointForFrame(FVector position)
+{
+#if WITH_WINDOWS_MIXED_REALITY
+	WindowsMixedReality::FWindowsMixedRealityHMD* hmd = GetWindowsMixedRealityHMD();
+	if (hmd == nullptr)
+	{
+		return;
+	}
+
+	hmd->SetFocusPointForFrame(position);
+#endif
+}
 
 // Temporary CVAR api for FocusPoint
 // The blueprint function to run a ConsoleCommand can be used to issue this each frame.
@@ -149,6 +135,8 @@ FPointerPoseInfo UWindowsMixedRealityFunctionLibrary::GetPointerPoseInfo(EContro
 // vr.SetFocusPointForFrame 102.4 -850.3 21.7
 static void SetFocusPointForFrame(const TArray<FString>& Args, UWorld*, FOutputDevice& Ar)
 {
+	Ar.Logf(ELogVerbosity::Error, TEXT("SetFocusPointForFrame DEPRECATED: please switch to using the SetFocusPointForFrame blueprint function in this library instead.  This command will be removed in 4.25."));
+
 	const int ArgsNum = Args.Num();
 	if (ArgsNum != 3)
 	{
@@ -185,6 +173,6 @@ static void SetFocusPointForFrame(const TArray<FString>& Args, UWorld*, FOutputD
 static FAutoConsoleCommand CSetFocusPointForFrameCmd(
 	TEXT("vr.SetFocusPointForFrame"),
 	*LOCTEXT("CVarText_SetFocusPointForFrame",
-		"Set the reference point for the stabilization plane on hololens 2. You must set it each frame to activate the feature for that frame.").ToString(),
+		"Set the reference point for the stabilization plane on hololens 2. You must set it each frame to activate the feature for that frame.  DEPRECATED: please switch to using the SetFocusPointForFrame blueprint function in this library instead.").ToString(),
 	FConsoleCommandWithWorldArgsAndOutputDeviceDelegate::CreateStatic(SetFocusPointForFrame));
 #undef LOCTEXT_NAMESPACE
\ No newline at end of file
