Change static MixedRealityInterop to use the existing heap-allocated version, which is nullptr checked to detect the case that we are on an unsupported version of windows.

diff --git a/Engine/Plugins/Runtime/WindowsMixedReality/Source/WindowsMixedRealityEyeTracker/Private/WindowsMixedRealityEyeTracker.cpp b/Engine/Plugins/Runtime/WindowsMixedReality/Source/WindowsMixedRealityEyeTracker/Private/WindowsMixedRealityEyeTracker.cpp
--- a/Engine/Plugins/Runtime/WindowsMixedReality/Source/WindowsMixedRealityEyeTracker/Private/WindowsMixedRealityEyeTracker.cpp
+++ b/Engine/Plugins/Runtime/WindowsMixedReality/Source/WindowsMixedRealityEyeTracker/Private/WindowsMixedRealityEyeTracker.cpp
@@ -8,9 +8,9 @@
 #if WITH_WINDOWS_MIXED_REALITY
 	#include "MixedRealityInterop.h"
 #endif
+#include "IWindowsMixedRealityHMDPlugin.h"
 
 #if WITH_WINDOWS_MIXED_REALITY
-static WindowsMixedReality::MixedRealityInterop hmd;
 
 class FWindowsMixedRealityEyeTracker :
 	public IEyeTracker
@@ -21,7 +21,11 @@ public:
 		// If this was created, then we want to use it, so request user perms
 #if PLATFORM_HOLOLENS
 		// If remoting, delay requesting permissions until after the remoting session is created.
-		hmd.RequestUserPermissionForEyeTracking();
+		WindowsMixedReality::MixedRealityInterop* HMD = IWindowsMixedRealityHMDPlugin::Get().GetMixedRealityInterop();
+		if (HMD != nullptr)
+		{
+			HMD->RequestUserPermissionForEyeTracking();
+		}
 #endif
 	}
 	
@@ -35,7 +39,8 @@ private:
 	virtual bool GetEyeTrackerGazeData(FEyeTrackerGazeData& OutGazeData) const override
 	{
 		WindowsMixedReality::EyeGazeRay ray;
-		if (!hmd.GetEyeGaze(ray))
+		WindowsMixedReality::MixedRealityInterop* HMD = IWindowsMixedRealityHMDPlugin::Get().GetMixedRealityInterop();
+		if (HMD == nullptr || !HMD->GetEyeGaze(ray))
 		{
 			return false;
 		}
@@ -51,7 +56,8 @@ private:
 	}
 	virtual EEyeTrackerStatus GetEyeTrackerStatus() const override
 	{
-		if (!hmd.SupportsEyeTracking() || !hmd.IsEyeTrackingAllowed())
+		WindowsMixedReality::MixedRealityInterop* HMD = IWindowsMixedRealityHMDPlugin::Get().GetMixedRealityInterop();
+		if (HMD == nullptr || !HMD->SupportsEyeTracking() || !HMD->IsEyeTrackingAllowed())
 		{
 			return EEyeTrackerStatus::NotConnected;
 		}
@@ -97,10 +103,13 @@ public:
 	virtual bool IsEyeTrackerConnected() const override
 	{
 #if WITH_WINDOWS_MIXED_REALITY
-		return hmd.SupportsEyeTracking();
-#else
-		return false;
+		WindowsMixedReality::MixedRealityInterop* HMD = IWindowsMixedRealityHMDPlugin::Get().GetMixedRealityInterop();
+		if (HMD != nullptr)
+		{
+			return HMD->SupportsEyeTracking();
+		}
 #endif
+		return false;
 	}
 	
 	virtual TSharedPtr< class IEyeTracker, ESPMode::ThreadSafe > CreateEyeTracker() override
diff --git a/Engine/Plugins/Runtime/WindowsMixedReality/Source/WindowsMixedRealityEyeTracker/WindowsMixedRealityEyeTracker.Build.cs b/Engine/Plugins/Runtime/WindowsMixedReality/Source/WindowsMixedRealityEyeTracker/WindowsMixedRealityEyeTracker.Build.cs
--- a/Engine/Plugins/Runtime/WindowsMixedReality/Source/WindowsMixedRealityEyeTracker/WindowsMixedRealityEyeTracker.Build.cs
+++ b/Engine/Plugins/Runtime/WindowsMixedReality/Source/WindowsMixedRealityEyeTracker/WindowsMixedRealityEyeTracker.Build.cs
@@ -27,6 +27,7 @@ namespace UnrealBuildTool.Rules
 						"MixedRealityInteropLibrary",
 						"InputDevice",
 						"EyeTracker",
+						"WindowsMixedRealityHMD",
 					}
 				);
 
diff --git a/Engine/Plugins/Runtime/WindowsMixedReality/Source/WindowsMixedRealityHMD/Private/WindowsMixedRealityHMD.cpp b/Engine/Plugins/Runtime/WindowsMixedReality/Source/WindowsMixedRealityHMD/Private/WindowsMixedRealityHMD.cpp
--- a/Engine/Plugins/Runtime/WindowsMixedReality/Source/WindowsMixedRealityHMD/Private/WindowsMixedRealityHMD.cpp
+++ b/Engine/Plugins/Runtime/WindowsMixedReality/Source/WindowsMixedRealityHMD/Private/WindowsMixedRealityHMD.cpp
@@ -106,6 +106,11 @@ namespace WindowsMixedReality
 	{
 		virtual TSharedPtr< class IXRTrackingSystem, ESPMode::ThreadSafe > CreateTrackingSystem() override;
 
+		virtual WindowsMixedReality::MixedRealityInterop* GetMixedRealityInterop() override
+		{
+			return HMD;
+		}
+
 		bool IsHMDConnected()
 		{
 #if WITH_WINDOWS_MIXED_REALITY
diff --git a/Engine/Plugins/Runtime/WindowsMixedReality/Source/WindowsMixedRealityHMD/Public/IWindowsMixedRealityHMDPlugin.h b/Engine/Plugins/Runtime/WindowsMixedReality/Source/WindowsMixedRealityHMD/Public/IWindowsMixedRealityHMDPlugin.h
--- a/Engine/Plugins/Runtime/WindowsMixedReality/Source/WindowsMixedRealityHMD/Public/IWindowsMixedRealityHMDPlugin.h
+++ b/Engine/Plugins/Runtime/WindowsMixedReality/Source/WindowsMixedRealityHMD/Public/IWindowsMixedRealityHMDPlugin.h
@@ -35,4 +35,9 @@ public:
 	{
 		return FModuleManager::Get().IsModuleLoaded( "WindowsMixedRealityHMD" );
 	}
+
+	virtual WindowsMixedReality::MixedRealityInterop* GetMixedRealityInterop()
+	{
+		return nullptr;
+	}
 };
