UEVR-1503 Hololens 2 - VRPreview window shows left half of the headset view.
-This bug also affected steamvr.
-Marked the Window as IsMirrorWindow, which causes the viewport size to be set from the window dimension rather than the size of the scene render target.  In VR the scene render target is sized for the HMD.  The desktop window may differ.
-Setup spectator screen for Hololens2 remoting.  This makes the various spectator screen modes available.

diff --git a/Engine/Plugins/Runtime/Steam/SteamVR/Source/SteamVR/Private/SteamVRHMD.cpp b/Engine/Plugins/Runtime/Steam/SteamVR/Source/SteamVR/Private/SteamVRHMD.cpp
--- a/Engine/Plugins/Runtime/Steam/SteamVR/Source/SteamVR/Private/SteamVRHMD.cpp
+++ b/Engine/Plugins/Runtime/Steam/SteamVR/Source/SteamVR/Private/SteamVRHMD.cpp
@@ -1353,6 +1353,9 @@ bool FSteamVRHMD::EnableStereo(bool bStereo)
 		TSharedPtr<SWindow> Window = SceneVP->FindWindow();
 		if (Window.IsValid() && SceneVP->GetViewportWidget().IsValid())
 		{
+			// Set MirrorWindow state on the Window
+			Window->SetMirrorWindow(bStereo);
+
 			if( bStereo )
 			{
 				uint32 Width, Height;
diff --git a/Engine/Plugins/Runtime/WindowsMixedReality/Source/WindowsMixedRealityHMD/Private/WindowsMixedRealityHMD.cpp b/Engine/Plugins/Runtime/WindowsMixedReality/Source/WindowsMixedRealityHMD/Private/WindowsMixedRealityHMD.cpp
--- a/Engine/Plugins/Runtime/WindowsMixedReality/Source/WindowsMixedRealityHMD/Private/WindowsMixedRealityHMD.cpp
+++ b/Engine/Plugins/Runtime/WindowsMixedReality/Source/WindowsMixedRealityHMD/Private/WindowsMixedRealityHMD.cpp
@@ -608,6 +608,11 @@ namespace WindowsMixedReality
 		InitTrackingFrame();
 		HMD->BlockUntilNextFrame();
 
+		if (SpectatorScreenController)
+		{
+			SpectatorScreenController->UpdateSpectatorScreenMode_RenderThread();
+		}
+
 		if (!HMD->CreateRenderingParameters())
 		{
 			// This will happen if an exception is thrown while creating the frame's rendering parameters.
@@ -617,6 +622,104 @@ namespace WindowsMixedReality
 #endif
 	}
 
+	FIntRect FWindowsMixedRealityHMD::GetFullFlatEyeRect_RenderThread(FTexture2DRHIRef EyeTexture) const
+	{
+		if (HMD->IsRemoting())
+		{
+			// Hololens has a relatively narrow FOV with very little distortion.
+			static FVector2D SrcNormRectMin(0.0f, 0.0f);
+			static FVector2D SrcNormRectMax(0.5f, 1.0f);
+			return FIntRect(EyeTexture->GetSizeX() * SrcNormRectMin.X, EyeTexture->GetSizeY() * SrcNormRectMin.Y, EyeTexture->GetSizeX() * SrcNormRectMax.X, EyeTexture->GetSizeY() * SrcNormRectMax.Y);
+		}
+		else
+		{
+			static FVector2D SrcNormRectMin(0.05f, 0.2f);
+			static FVector2D SrcNormRectMax(0.45f, 0.8f);
+			return FIntRect(EyeTexture->GetSizeX() * SrcNormRectMin.X, EyeTexture->GetSizeY() * SrcNormRectMin.Y, EyeTexture->GetSizeX() * SrcNormRectMax.X, EyeTexture->GetSizeY() * SrcNormRectMax.Y);
+		}
+	}
+
+	void FWindowsMixedRealityHMD::CopyTexture_RenderThread(FRHICommandListImmediate& RHICmdList, FRHITexture2D* SrcTexture, FIntRect SrcRect, FRHITexture2D* DstTexture, FIntRect DstRect, bool bClearBlack, bool bNoAlpha) const
+	{
+		check(IsInRenderingThread());
+
+		const uint32 ViewportWidth = DstRect.Width();
+		const uint32 ViewportHeight = DstRect.Height();
+		const FIntPoint TargetSize(ViewportWidth, ViewportHeight);
+
+		const float SrcTextureWidth = SrcTexture->GetSizeX();
+		const float SrcTextureHeight = SrcTexture->GetSizeY();
+		float U = 0.f, V = 0.f, USize = 1.f, VSize = 1.f;
+		if (!SrcRect.IsEmpty())
+		{
+			U = SrcRect.Min.X / SrcTextureWidth;
+			V = SrcRect.Min.Y / SrcTextureHeight;
+			USize = SrcRect.Width() / SrcTextureWidth;
+			VSize = SrcRect.Height() / SrcTextureHeight;
+		}
+
+		// #todo-renderpasses Possible optimization here - use DontLoad if we will immediately clear the entire target
+		FRHIRenderPassInfo RPInfo(DstTexture, ERenderTargetActions::Load_Store);
+		RHICmdList.BeginRenderPass(RPInfo, TEXT("WindowsMixedRealityHMD_CopyTexture"));
+		{
+			if (bClearBlack)
+			{
+				const FIntRect ClearRect(0, 0, DstTexture->GetSizeX(), DstTexture->GetSizeY());
+				RHICmdList.SetViewport(ClearRect.Min.X, ClearRect.Min.Y, 0, ClearRect.Max.X, ClearRect.Max.Y, 1.0f);
+				DrawClearQuad(RHICmdList, FLinearColor::Black);
+			}
+
+			RHICmdList.SetViewport(DstRect.Min.X, DstRect.Min.Y, 0, DstRect.Max.X, DstRect.Max.Y, 1.0f);
+
+			FGraphicsPipelineStateInitializer GraphicsPSOInit;
+			RHICmdList.ApplyCachedRenderTargets(GraphicsPSOInit);
+			GraphicsPSOInit.BlendState = bNoAlpha ? TStaticBlendState<>::GetRHI() : TStaticBlendState<CW_RGBA, BO_Add, BF_SourceAlpha, BF_InverseSourceAlpha, BO_Add, BF_One, BF_InverseSourceAlpha>::GetRHI();
+			GraphicsPSOInit.RasterizerState = TStaticRasterizerState<>::GetRHI();
+			GraphicsPSOInit.DepthStencilState = TStaticDepthStencilState<false, CF_Always>::GetRHI();
+			GraphicsPSOInit.PrimitiveType = PT_TriangleList;
+
+			const auto FeatureLevel = GMaxRHIFeatureLevel;
+			auto ShaderMap = GetGlobalShaderMap(FeatureLevel);
+
+			TShaderMapRef<FScreenVS> VertexShader(ShaderMap);
+			GraphicsPSOInit.BoundShaderState.VertexDeclarationRHI = GFilterVertexDeclaration.VertexDeclarationRHI;
+			GraphicsPSOInit.BoundShaderState.VertexShaderRHI = GETSAFERHISHADER_VERTEX(*VertexShader);
+
+			const bool bSameSize = DstRect.Size() == SrcRect.Size();
+			FRHISamplerState* PixelSampler = bSameSize ? TStaticSamplerState<SF_Point>::GetRHI() : TStaticSamplerState<SF_Bilinear>::GetRHI();
+
+			if ((SrcTexture->GetFlags() & TexCreate_SRGB) != 0)
+			{
+				TShaderMapRef<FScreenPSsRGBSource> PixelShader(ShaderMap);
+				GraphicsPSOInit.BoundShaderState.PixelShaderRHI = GETSAFERHISHADER_PIXEL(*PixelShader);
+
+				SetGraphicsPipelineState(RHICmdList, GraphicsPSOInit);
+				PixelShader->SetParameters(RHICmdList, PixelSampler, SrcTexture);
+			}
+			else
+			{
+				TShaderMapRef<FScreenPS> PixelShader(ShaderMap);
+				GraphicsPSOInit.BoundShaderState.PixelShaderRHI = GETSAFERHISHADER_PIXEL(*PixelShader);
+
+				SetGraphicsPipelineState(RHICmdList, GraphicsPSOInit);
+				PixelShader->SetParameters(RHICmdList, PixelSampler, SrcTexture);
+			}
+
+			RendererModule->DrawRectangle(
+				RHICmdList,
+				0, 0,
+				ViewportWidth, ViewportHeight,
+				U, V,
+				USize, VSize,
+				TargetSize,
+				FIntPoint(1, 1),
+				*VertexShader,
+				EDRF_Default);
+
+			RHICmdList.EndRenderPass();
+		}
+	}
+
 	bool FWindowsMixedRealityHMD::GetCurrentPose(
 		int32 DeviceId,
 		FQuat& CurrentOrientation,
@@ -820,6 +923,11 @@ namespace WindowsMixedReality
 			TSharedPtr<SWindow> Window = SceneVP->FindWindow();
 			//if (Window.IsValid() && SceneVP->GetViewportWidget().IsValid())
 			{
+#if !PLATFORM_HOLOLENS
+				// Set MirrorWindow state on the Window
+				Window->SetMirrorWindow(bIsStereoDesired);
+#endif
+
 				if (bIsStereoDesired)
 				{
 					int Width, Height;
@@ -926,19 +1034,37 @@ namespace WindowsMixedReality
 
 	FMatrix FWindowsMixedRealityHMD::GetStereoProjectionMatrix(const enum EStereoscopicPass StereoPassType) const
 	{
-		check(IsInGameThread());
-
-		if (StereoPassType == eSSP_LEFT_EYE)
-		{
-			return Frame_GameThread.ProjectionMatrixL;
-		}
-		else if (StereoPassType == eSSP_RIGHT_EYE)
+		if (IsInRenderingThread())
 		{
-			return Frame_GameThread.ProjectionMatrixR;
+			if (StereoPassType == eSSP_LEFT_EYE)
+			{
+				return Frame_RenderThread.ProjectionMatrixL;
+			}
+			else if (StereoPassType == eSSP_RIGHT_EYE)
+			{
+				return Frame_RenderThread.ProjectionMatrixR;
+			}
+			else
+			{
+				return FMatrix::Identity;
+			}
 		}
 		else
 		{
-			return FMatrix::Identity;
+			check(IsInGameThread());
+
+			if (StereoPassType == eSSP_LEFT_EYE)
+			{
+				return Frame_GameThread.ProjectionMatrixL;
+			}
+			else if (StereoPassType == eSSP_RIGHT_EYE)
+			{
+				return Frame_GameThread.ProjectionMatrixR;
+			}
+			else
+			{
+				return FMatrix::Identity;
+			}		
 		}
 	}
 
@@ -1006,84 +1132,11 @@ namespace WindowsMixedReality
 
 	void FWindowsMixedRealityHMD::RenderTexture_RenderThread(class FRHICommandListImmediate& RHICmdList, class FRHITexture2D* BackBuffer, class FRHITexture2D* SrcTexture, FVector2D WindowSize) const
 	{
-#if PLATFORM_HOLOLENS
-		return;
-#endif
-		const uint32 WindowWidth = gameWindowWidth;
-		const uint32 WindowHeight = gameWindowHeight;
-
-		const uint32 ViewportWidth = BackBuffer->GetSizeX();
-		const uint32 ViewportHeight = BackBuffer->GetSizeY();
-
-		const uint32 TextureWidth = SrcTexture->GetSizeX();
-		const uint32 TextureHeight = SrcTexture->GetSizeY();
-
-		const uint32 SourceWidth = TextureWidth / 2;
-		const uint32 SourceHeight = TextureHeight;
-
-		const float r = (float)SourceWidth / (float)SourceHeight;
-
-		float width = (float)WindowWidth;
-		float height = (float)WindowHeight;
-
-		if ((float)WindowWidth / r < WindowHeight)
-		{
-			width = ViewportWidth;
-
-			float displayHeight = (float)WindowWidth / r;
-			height = (float)ViewportHeight * (displayHeight / (float)WindowHeight);
-		}
-		else // width > height
+		if (SpectatorScreenController)
 		{
-			height = ViewportHeight;
-
-			float displayWidth = (float)WindowHeight * r;
-			width = (float)ViewportWidth * (displayWidth / (float)WindowWidth);
+			SpectatorScreenController->RenderSpectatorScreen_RenderThread(RHICmdList, BackBuffer, SrcTexture, WindowSize);
 		}
 
-		width = FMath::Clamp<int>(width, 10, ViewportWidth);
-		height = FMath::Clamp<int>(height, 10, ViewportHeight);
-
-		const uint32 x = (ViewportWidth - width) * 0.5f;
-		const uint32 y = (ViewportHeight - height) * 0.5f;
-		FRHIRenderPassInfo RPInfo(BackBuffer, ERenderTargetActions::Load_Store);
-		RHICmdList.BeginRenderPass(RPInfo, TEXT("WindowsMixedRealityHMD"));
-		DrawClearQuad(RHICmdList, FLinearColor(0.0f, 0.0f, 0.0f, 1.0f));
-		RHICmdList.SetViewport(x, y, 0, width + x, height + y, 1.0f);
-
-		FGraphicsPipelineStateInitializer GraphicsPSOInit;
-		RHICmdList.ApplyCachedRenderTargets(GraphicsPSOInit);
-
-		const auto FeatureLevel = GMaxRHIFeatureLevel;
-		auto ShaderMap = GetGlobalShaderMap(FeatureLevel);
-		TShaderMapRef<FScreenVS> VertexShader(ShaderMap);
-		TShaderMapRef<FScreenPS> PixelShader(ShaderMap);
-
-		GraphicsPSOInit.BlendState = TStaticBlendState<>::GetRHI();
-		GraphicsPSOInit.RasterizerState = TStaticRasterizerState<>::GetRHI();
-		GraphicsPSOInit.DepthStencilState = TStaticDepthStencilState<false, CF_Always>::GetRHI();
-		GraphicsPSOInit.BoundShaderState.VertexDeclarationRHI = GFilterVertexDeclaration.VertexDeclarationRHI;
-		GraphicsPSOInit.BoundShaderState.VertexShaderRHI = GETSAFERHISHADER_VERTEX(*VertexShader);
-		GraphicsPSOInit.BoundShaderState.PixelShaderRHI = GETSAFERHISHADER_PIXEL(*PixelShader);
-		GraphicsPSOInit.PrimitiveType = PT_TriangleList;
-
-		SetGraphicsPipelineState(RHICmdList, GraphicsPSOInit);
-
-		PixelShader->SetParameters(RHICmdList, TStaticSamplerState<SF_Bilinear>::GetRHI(), SrcTexture);
-
-		RendererModule->DrawRectangle(
-			RHICmdList,
-			0, 0,
-			ViewportWidth, ViewportHeight,
-			0.0f, 0.0f,
-			0.5f, 1.0f,
-			FIntPoint(ViewportWidth, ViewportHeight),
-			FIntPoint(1, 1),
-			*VertexShader,
-			EDRF_Default);
-
-		RHICmdList.EndRenderPass();
-
 		// We keep refs to the depth texture in the hmd, so this function is non-const.
 		// RenderTexture_RenderThread should perhaps be refactored or made non-const.  
 		// But to get this into a hotfix of 4.23 we shall simply cast.
@@ -1421,6 +1474,10 @@ namespace WindowsMixedReality
 
 		HiddenAreaMesh.SetNum(2);
 		VisibleAreaMesh.SetNum(2);
+
+#if !PLATFORM_HOLOLENS
+		CreateSpectatorScreenController();
+#endif
 	}
 #endif
 
@@ -1813,6 +1870,11 @@ namespace WindowsMixedReality
 #endif
 #endif
 	}
+
+	void FWindowsMixedRealityHMD::CreateSpectatorScreenController()
+	{
+		SpectatorScreenController = MakeUnique<FDefaultSpectatorScreenController>(this);
+	}
 }
 
 IMPLEMENT_MODULE(WindowsMixedReality::FWindowsMixedRealityHMDPlugin, WindowsMixedRealityHMD)
diff --git a/Engine/Plugins/Runtime/WindowsMixedReality/Source/WindowsMixedRealityHMD/Private/WindowsMixedRealityHMD.h b/Engine/Plugins/Runtime/WindowsMixedReality/Source/WindowsMixedRealityHMD/Private/WindowsMixedRealityHMD.h
--- a/Engine/Plugins/Runtime/WindowsMixedReality/Source/WindowsMixedRealityHMD/Private/WindowsMixedRealityHMD.h
+++ b/Engine/Plugins/Runtime/WindowsMixedReality/Source/WindowsMixedRealityHMD/Private/WindowsMixedRealityHMD.h
@@ -126,6 +126,11 @@ namespace WindowsMixedReality
 		virtual bool HasVisibleAreaMesh() const override;
 		virtual void DrawVisibleAreaMesh_RenderThread(FRHICommandList& RHICmdList, EStereoscopicPass StereoPass) const override;
 
+		// Spectator screen Hooks.
+		virtual FIntRect GetFullFlatEyeRect_RenderThread(FTexture2DRHIRef EyeTexture) const override;
+		// Helper to copy one render target into another for spectator screen display
+		virtual void CopyTexture_RenderThread(FRHICommandListImmediate& RHICmdList, FRHITexture2D* SrcTexture, FIntRect SrcRect, FRHITexture2D* DstTexture, FIntRect DstRect, bool bClearBlack, bool bNoAlpha) const override;
+
 		/** ISceneViewExtension interface */
 		virtual void SetupViewFamily(FSceneViewFamily& InViewFamily) override;
 		virtual void SetupView(
@@ -333,6 +338,9 @@ namespace WindowsMixedReality
 			return new SpeechRecognizerInterop();
 		}
 #endif
+
+	private:
+		void CreateSpectatorScreenController();
 	};
 }
 
