Force load the dependent DLLs from a known correct location earlier than the MixedRealityInterop loads so it doesn't try to look for them in the wrong place.

diff --git a/Engine/Plugins/Runtime/WindowsMixedReality/Source/WindowsMixedRealityHMD/Private/WindowsMixedRealityHMD.cpp b/Engine/Plugins/Runtime/WindowsMixedReality/Source/WindowsMixedRealityHMD/Private/WindowsMixedRealityHMD.cpp
--- a/Engine/Plugins/Runtime/WindowsMixedReality/Source/WindowsMixedRealityHMD/Private/WindowsMixedRealityHMD.cpp
+++ b/Engine/Plugins/Runtime/WindowsMixedReality/Source/WindowsMixedRealityHMD/Private/WindowsMixedRealityHMD.cpp
@@ -155,8 +155,14 @@ namespace WindowsMixedReality
 				// Load these dependencies first or MixedRealityInteropLibraryHandle fails to load since it doesn't look in the correct path for its dependencies automatically
 				FString HoloLensLibraryDir = EngineDir / "Binaries/ThirdParty/Windows/x64";
 				FPlatformProcess::PushDllDirectory(*HoloLensLibraryDir);
-				FPlatformProcess::GetDllHandle(_TEXT("PerceptionDevice.dll"));
-				FPlatformProcess::GetDllHandle(_TEXT("Microsoft.Holographic.AppRemoting.dll"));
+				FPlatformProcess::GetDllHandle(TEXT("PerceptionDevice.dll"));
+				FPlatformProcess::GetDllHandle(TEXT("Microsoft.Holographic.AppRemoting.dll"));
+				FPlatformProcess::PopDllDirectory(*HoloLensLibraryDir);
+
+				HoloLensLibraryDir = EngineDir / "Binaries/Win64";
+				FPlatformProcess::PushDllDirectory(*HoloLensLibraryDir);
+				FPlatformProcess::GetDllHandle(TEXT("Microsoft.Perception.Simulation.dll"));
+				FPlatformProcess::GetDllHandle(TEXT("HolographicStreamerDesktop.dll"));
 				FPlatformProcess::PopDllDirectory(*HoloLensLibraryDir);
 #endif // PLATFORM_64BITS && WITH_EDITOR
 
