Fix WMR thumbstick and touchpad directional keys

diff --git a/Engine/Plugins/Runtime/WindowsMixedReality/Source/WindowsMixedRealitySpatialInput/Private/WindowsMixedRealitySpatialInput.cpp b/Engine/Plugins/Runtime/WindowsMixedReality/Source/WindowsMixedRealitySpatialInput/Private/WindowsMixedRealitySpatialInput.cpp
--- a/Engine/Plugins/Runtime/WindowsMixedReality/Source/WindowsMixedRealitySpatialInput/Private/WindowsMixedRealitySpatialInput.cpp
+++ b/Engine/Plugins/Runtime/WindowsMixedReality/Source/WindowsMixedRealitySpatialInput/Private/WindowsMixedRealitySpatialInput.cpp
@@ -288,8 +288,12 @@ namespace WindowsMixedReality
 			if (position.Size() > 0.5f)
 			{
 				key = GetThumbstickDirectionKey(position, hand);
-				if (!ThumbstickDirection[i].IsValid())
+				if (key.IsValid())
 				{
+					if (ThumbstickDirection[i].IsValid() && key != ThumbstickDirection[i])
+					{
+						SendControllerButtonEvent(MessageHandler, source, ThumbstickDirection[i], HMDInputPressState::Released);
+					}
 					SendControllerButtonEvent(MessageHandler, source, key, HMDInputPressState::Pressed);
 				}
 				ThumbstickDirection[i] = key;
@@ -335,19 +339,23 @@ namespace WindowsMixedReality
 			if (position.Size() > 0.5f)
 			{
 				key = GetTouchpadDirectionKey(position, hand);
-				if (!ThumbstickDirection[i].IsValid())
+				if (key.IsValid())
 				{
+					if (TouchpadDirection[i].IsValid() && key != TouchpadDirection[i])
+					{
+						SendControllerButtonEvent(MessageHandler, source, TouchpadDirection[i], HMDInputPressState::Released);
+					}
 					SendControllerButtonEvent(MessageHandler, source, key, HMDInputPressState::Pressed);
 				}
-				ThumbstickDirection[i] = key;
+				TouchpadDirection[i] = key;
 			}
 			else
 			{
-				if (ThumbstickDirection[i].IsValid())
+				if (TouchpadDirection[i].IsValid())
 				{
-					SendControllerButtonEvent(MessageHandler, source, ThumbstickDirection[i], HMDInputPressState::Released);
+					SendControllerButtonEvent(MessageHandler, source, TouchpadDirection[i], HMDInputPressState::Released);
 				}
-				ThumbstickDirection[i] = EKeys::Invalid;
+				TouchpadDirection[i] = EKeys::Invalid;
 			}
 		}
 	}
