So investigating this, it looks like the problem is that FWindowsMixedRealityHMD::SetupHolographicCamera() is setting the window to be a mirror window and it's viewport's size to not be driven by the window, but then never resetting those when we exit stereo.
This change seems to alleviate the issue for me

diff --git a/Engine/Plugins/Runtime/WindowsMixedReality/Source/WindowsMixedRealityHMD/Private/WindowsMixedRealityHMD.cpp b/Engine/Plugins/Runtime/WindowsMixedReality/Source/WindowsMixedRealityHMD/Private/WindowsMixedRealityHMD.cpp
--- a/Engine/Plugins/Runtime/WindowsMixedReality/Source/WindowsMixedRealityHMD/Private/WindowsMixedRealityHMD.cpp
+++ b/Engine/Plugins/Runtime/WindowsMixedReality/Source/WindowsMixedRealityHMD/Private/WindowsMixedRealityHMD.cpp
@@ -1522,6 +1522,24 @@ namespace WindowsMixedReality
 
 		bIsStereoDesired = false;
 		bIsStereoEnabled = false;
+
+#if WITH_WINDOWS_MIXED_REALITY
+		// Reset the viewport to no longer match the HMD display
+		FSceneViewport* SceneVP = FindMRSceneViewport(bIsStereoDesired);
+
+		if (SceneVP)
+		{
+			TSharedPtr<SWindow> Window = SceneVP->FindWindow();
+			if (Window.IsValid())
+			{
+#if !PLATFORM_HOLOLENS
+				// Set MirrorWindow state on the Window
+				Window->SetMirrorWindow(bIsStereoDesired);
+				Window->SetViewportSizeDrivenByWindow(true);
+#endif
+			}
+		}
+#endif
 		
 		bNeedReallocateDepthTexture = false;
 
