Submitting two files missing from CL#10288280.
Fixes compilation after stereo layers refactoring

diff --git a/Engine/Plugins/Lumin/MagicLeap/Source/MagicLeap/Private/MagicLeapStereoLayers.cpp b/Engine/Plugins/Lumin/MagicLeap/Source/MagicLeap/Private/MagicLeapStereoLayers.cpp
--- a/Engine/Plugins/Lumin/MagicLeap/Source/MagicLeap/Private/MagicLeapStereoLayers.cpp
+++ b/Engine/Plugins/Lumin/MagicLeap/Source/MagicLeap/Private/MagicLeapStereoLayers.cpp
@@ -80,7 +80,6 @@ IStereoLayers::FLayerDesc FMagicLeapStereoLayers::GetDebugCanvasLayerDesc(FTextu
 	const float DebugCanvasHeight = CVarHeight ? CVarHeight->GetFloat() : DefaultHeight;
 	StereoLayerDesc.QuadSize = FVector2D(DebugCanvasWidth, DebugCanvasHeight);
 	StereoLayerDesc.PositionType = IStereoLayers::ELayerType::FaceLocked;
-	StereoLayerDesc.ShapeType = IStereoLayers::ELayerShape::QuadLayer;
 	StereoLayerDesc.Texture = Texture;
 	StereoLayerDesc.Flags = IStereoLayers::ELayerFlags::LAYER_FLAG_TEX_CONTINUOUS_UPDATE;
 	StereoLayerDesc.Flags |= IStereoLayers::ELayerFlags::LAYER_FLAG_QUAD_PRESERVE_TEX_RATIO;
diff --git a/Engine/Plugins/Runtime/WindowsMixedReality/Source/WindowsMixedRealityHMD/Private/WindowsMixedRealityStereoLayers.cpp b/Engine/Plugins/Runtime/WindowsMixedReality/Source/WindowsMixedRealityHMD/Private/WindowsMixedRealityStereoLayers.cpp
--- a/Engine/Plugins/Runtime/WindowsMixedReality/Source/WindowsMixedRealityHMD/Private/WindowsMixedRealityStereoLayers.cpp
+++ b/Engine/Plugins/Runtime/WindowsMixedReality/Source/WindowsMixedRealityHMD/Private/WindowsMixedRealityStereoLayers.cpp
@@ -26,7 +26,6 @@ IStereoLayers::FLayerDesc FWindowsMixedRealityStereoLayers::GetDebugCanvasLayerD
 	StereoLayerDesc.QuadSize = FVector2D(120.f, 120.f);
 #endif
 	StereoLayerDesc.PositionType = IStereoLayers::ELayerType::FaceLocked;
-	StereoLayerDesc.ShapeType = IStereoLayers::ELayerShape::QuadLayer;
 	StereoLayerDesc.Texture = Texture;
 	StereoLayerDesc.Flags = IStereoLayers::ELayerFlags::LAYER_FLAG_TEX_CONTINUOUS_UPDATE;
 	StereoLayerDesc.Flags |= IStereoLayers::ELayerFlags::LAYER_FLAG_QUAD_PRESERVE_TEX_RATIO;
