DRPS - Compile fix
#rnx

diff --git a/Engine/Plugins/Runtime/WindowsMixedReality/Source/WindowsMixedRealityRHI/Private/WindowsMixedRealityRHI.cpp b/Engine/Plugins/Runtime/WindowsMixedReality/Source/WindowsMixedRealityRHI/Private/WindowsMixedRealityRHI.cpp
--- a/Engine/Plugins/Runtime/WindowsMixedReality/Source/WindowsMixedRealityRHI/Private/WindowsMixedRealityRHI.cpp
+++ b/Engine/Plugins/Runtime/WindowsMixedReality/Source/WindowsMixedRealityRHI/Private/WindowsMixedRealityRHI.cpp
@@ -261,7 +261,7 @@ void FWindowsMixedRealityViewport::UpdateBackBuffer()
 void FWindowsMixedRealityViewport::Resize(uint32 InSizeX, uint32 InSizeY, bool bInIsFullscreen, EPixelFormat PreferredPixelFormat)
 {
 	// Unbind any dangling references to resources
-	D3DRHI->RHISetRenderTargets(0, nullptr, nullptr, 0, nullptr);
+	D3DRHI->RHISetRenderTargets(0, nullptr, nullptr);
 	D3DRHI->ClearState();
 	D3DRHI->GetDeviceContext()->Flush(); // Potential perf hit
 
