Fix CIS static analysis error (not checking if object null before passing to copySubresourceRegion)

diff --git a/Engine/Plugins/Runtime/WindowsMixedReality/Source/WindowsMixedRealityHMD/Private/WindowsMixedRealityHMD.cpp b/Engine/Plugins/Runtime/WindowsMixedReality/Source/WindowsMixedRealityHMD/Private/WindowsMixedRealityHMD.cpp
--- a/Engine/Plugins/Runtime/WindowsMixedReality/Source/WindowsMixedRealityHMD/Private/WindowsMixedRealityHMD.cpp
+++ b/Engine/Plugins/Runtime/WindowsMixedReality/Source/WindowsMixedRealityHMD/Private/WindowsMixedRealityHMD.cpp
@@ -1550,7 +1550,7 @@ namespace WindowsMixedReality
 
 		HMD->CommitDepthBuffer(stereoDepthTexture);
 		// Third camera depth
-		if (HMD->IsThirdCameraActive() && remappedDepthTexture != nullptr)
+		if (HMD->IsThirdCameraActive() && remappedDepthTexture != nullptr && monoDepthTexture != nullptr)
 		{
 			ID3D11Texture2D* sourceDepth = (ID3D11Texture2D*)remappedDepthTexture->GetNativeResource();
 			int w, h;
