Change: 11467954
UE-80182 HoloLens(1&2) Failure to connect via holographic remoting
-Added a warning if you are running with the wrong XRSystem active (eg oculus) and trying to remote to a hololens (which requires that your xrsystem be WindowsMixedRealityHMD) so that we don't just silently not do anything when you click the 'connect' button.

diff --git a/Engine/Plugins/Runtime/WindowsMixedReality/Source/WindowsMixedRealityHMD/Private/WindowsMixedRealityStatics.cpp b/Engine/Plugins/Runtime/WindowsMixedReality/Source/WindowsMixedRealityHMD/Private/WindowsMixedRealityStatics.cpp
--- a/Engine/Plugins/Runtime/WindowsMixedReality/Source/WindowsMixedRealityHMD/Private/WindowsMixedRealityStatics.cpp
+++ b/Engine/Plugins/Runtime/WindowsMixedReality/Source/WindowsMixedRealityHMD/Private/WindowsMixedRealityStatics.cpp
@@ -5,6 +5,10 @@
 #include "Engine/Engine.h"
 #include "WindowsMixedRealityHMD.h"
 
+#if WITH_EDITOR
+#include "WindowsMixedRealityRuntimeSettings.h"
+#endif
+
 namespace WindowsMixedReality
 {
 	FWindowsMixedRealityHMD* GetWindowsMixedRealityHMD() noexcept
@@ -151,10 +155,26 @@ namespace WindowsMixedReality
 		}
 		else
 		{
-			UE_LOG(LogWmrHmd, Log, TEXT("FWindowsMixedRealityStatics::ConnectToRemoteHoloLens() is doing nothing because FWindowsMixedRealityHMD fetch failed.  Perhaps you need to enable remoting and restart the editor?"));
+			if (!GEngine->XRSystem.IsValid())
+			{
+#if WITH_EDITOR
+				UWindowsMixedRealityRuntimeSettings::Get()->OnRemotingStatusChanged.ExecuteIfBound(FString(TEXT("Cannot Connect, see log for Errors.")), FLinearColor::Red);
+#endif
+				UE_LOG(LogWmrHmd, Error, TEXT("ConnectToRemoteHoloLens XRSystem is not valid. Perhaps it failed to start up?  Cannot Connect."));
+			}
+			else if (GEngine->XRSystem->GetSystemName() != FName("WindowsMixedRealityHMD"))
+			{
+#if WITH_EDITOR
+				UWindowsMixedRealityRuntimeSettings::Get()->OnRemotingStatusChanged.ExecuteIfBound(FString(TEXT("Cannot Connect, see log for Errors.")), FLinearColor::Red);
+#endif
+				UE_LOG(LogWmrHmd, Error, TEXT("ConnectToRemoteHoloLens XRSystem SystemName is %s, not WindowsMixedRealityHMD.  Cannot Connect.  Perhaps you want to disable other XR plugins, adjust the priorities of XR plugins, deactivate other XR hardware, or run with -hmd=WindowsMixedRealityHMD in your editor commandline?"), *GEngine->XRSystem->GetSystemName().ToString());
+			}
 		}
 #else
-		UE_LOG(LogWmrHmd, Log, TEXT("FWindowsMixedRealityStatics::ConnectToRemoteHoloLens() is doing nothing because PLATFORM_HOLOLENS. (You don't 'remote' when running on device or emulated device.)"));
+#if WITH_EDITOR
+		UWindowsMixedRealityRuntimeSettings::Get()->OnRemotingStatusChanged.ExecuteIfBound(FString(TEXT("Cannot Connect, see log for Errors.")), FLinearColor::Red);
+#endif
+		UE_LOG(LogWmrHmd, Error, TEXT("FWindowsMixedRealityStatics::ConnectToRemoteHoloLens() is doing nothing because PLATFORM_HOLOLENS. (You don't 'remote' when running on device or emulated device.)"));
 #endif
 	}
 
