UE-88906 //UE4/Release-4.25 - Build MissionAR Win64 - Fatal error: Couldn't find source file of virtual shader path '/Plugin/WindowsMixedReality/Private/DepthConversion.usf'
-Moved AddShaderSourceDirectoryMapping up so that it happens even if we are running on an unsupported windows version where we will not load the wmr interop shader compliation can succeed.

diff --git a/Engine/Plugins/Runtime/WindowsMixedReality/Source/WindowsMixedRealityHMD/Private/WindowsMixedRealityHMD.cpp b/Engine/Plugins/Runtime/WindowsMixedReality/Source/WindowsMixedRealityHMD/Private/WindowsMixedRealityHMD.cpp
--- a/Engine/Plugins/Runtime/WindowsMixedReality/Source/WindowsMixedRealityHMD/Private/WindowsMixedRealityHMD.cpp
+++ b/Engine/Plugins/Runtime/WindowsMixedReality/Source/WindowsMixedRealityHMD/Private/WindowsMixedRealityHMD.cpp
@@ -124,6 +124,10 @@ namespace WindowsMixedReality
 #if WITH_WINDOWS_MIXED_REALITY
 			IHeadMountedDisplayModule::StartupModule();
 
+			// Set the shader directory even if we won't be able to load the interop so shader compliation does not fail.
+			FString PluginShaderDir = FPaths::Combine(IPluginManager::Get().FindPlugin(TEXT("WindowsMixedReality"))->GetBaseDir(), TEXT("Shaders"));
+			AddShaderSourceDirectoryMapping(TEXT("/Plugin/WindowsMixedReality"), PluginShaderDir);
+
 			HMD = LoadInteropLibrary();
 			if (!HMD)
 			{
@@ -137,8 +141,6 @@ namespace WindowsMixedReality
 			}	
 #endif // WANTS_INTEROP_LOGGING
 
-			FString PluginShaderDir = FPaths::Combine(IPluginManager::Get().FindPlugin(TEXT("WindowsMixedReality"))->GetBaseDir(), TEXT("Shaders"));
-			AddShaderSourceDirectoryMapping(TEXT("/Plugin/WindowsMixedReality"), PluginShaderDir);
 #else // WITH_WINDOWS_MIXED_REALITY
 			UE_LOG(LogWmrHmd, Error, TEXT("Windows Mixed Reality compiled with unsupported compiler.  Please recompile with Visual Studio 2017"));
 #endif // WITH_WINDOWS_MIXED_REALITY
