Untested fix (requires windows server machine to test).  Only checking for the "Windows Server Technical Preview" because that's what we have on our build machines for hololens it seems.

diff --git a/Engine/Plugins/Runtime/WindowsMixedReality/Source/WindowsMixedRealityHMD/Private/WindowsMixedRealityInteropLoader.cpp b/Engine/Plugins/Runtime/WindowsMixedReality/Source/WindowsMixedRealityHMD/Private/WindowsMixedRealityInteropLoader.cpp
--- a/Engine/Plugins/Runtime/WindowsMixedReality/Source/WindowsMixedRealityHMD/Private/WindowsMixedRealityInteropLoader.cpp
+++ b/Engine/Plugins/Runtime/WindowsMixedReality/Source/WindowsMixedRealityHMD/Private/WindowsMixedRealityInteropLoader.cpp
@@ -27,6 +27,13 @@ namespace WindowsMixedReality
 		OSSubVersionLabel = OSVersionLabel;
 		bool bHasSupportedWindowsVersion = OSSubVersionLabel.RemoveFromStart("Windows 10 (Release ") && OSSubVersionLabel.RemoveFromEnd(")") && (FCString::Atoi(*OSSubVersionLabel) >= MIN_WIN_10_VERSION_FOR_WMR);
 
+		// If we can't find Win10 version, check for Windows Server equivalent
+		if (!bHasSupportedWindowsVersion)
+		{
+			OSSubVersionLabel = OSVersionLabel;
+			bHasSupportedWindowsVersion = OSSubVersionLabel.RemoveFromStart("Windows Server Technical Preview (Release ") && OSSubVersionLabel.RemoveFromEnd(")") && (FCString::Atoi(*OSSubVersionLabel) >= MIN_WIN_10_VERSION_FOR_WMR);
+		}
+
 		if (bHasSupportedWindowsVersion)
 		{
 			// Get the base directory of this plugin
