WMR input simulation cleanup and fixes (provided by Microsoft): Additional features and fixes for WMR input simulation support

diff --git a/Engine/Plugins/Runtime/WindowsMixedReality/Source/WindowsMixedRealityHMD/Private/WindowsMixedRealityHMD.cpp b/Engine/Plugins/Runtime/WindowsMixedReality/Source/WindowsMixedRealityHMD/Private/WindowsMixedRealityHMD.cpp
--- a/Engine/Plugins/Runtime/WindowsMixedReality/Source/WindowsMixedRealityHMD/Private/WindowsMixedRealityHMD.cpp
+++ b/Engine/Plugins/Runtime/WindowsMixedReality/Source/WindowsMixedRealityHMD/Private/WindowsMixedRealityHMD.cpp
@@ -211,6 +211,23 @@ namespace WindowsMixedReality
 	// FWindowsMixedRealityHMD IHeadMountedDisplay Implementation
 	//---------------------------------------------------
 
+	bool FWindowsMixedRealityHMD::IsHeadTrackingAllowed() const
+	{
+		if (FHeadMountedDisplayBase::IsHeadTrackingAllowed())
+		{
+			return true;
+		}
+
+#if WITH_INPUT_SIMULATION
+		if (auto* InputSim = UWindowsMixedRealityInputSimulationEngineSubsystem::GetInputSimulationIfEnabled())
+		{
+			return true;
+		}
+#endif
+
+		return false;
+	}
+
 	bool FWindowsMixedRealityHMD::IsHMDConnected()
 	{
 #if WITH_WINDOWS_MIXED_REALITY
@@ -343,8 +360,17 @@ namespace WindowsMixedReality
 
 	bool FWindowsMixedRealityHMD::HasValidTrackingPosition()
 	{
-		const Frame& TheFrame = GetFrame();
-		return TheFrame.bPositionalTrackingUsed;
+#if WITH_INPUT_SIMULATION
+		if (auto* InputSim = UWindowsMixedRealityInputSimulationEngineSubsystem::GetInputSimulationIfEnabled())
+		{
+			return InputSim->HasPositionalTracking();
+		}
+		else
+#endif
+		{
+			const Frame& TheFrame = GetFrame();
+			return TheFrame.bPositionalTrackingUsed;
+		}
 	}
 
 	FString FWindowsMixedRealityHMD::GetVersionString() const
@@ -710,10 +736,20 @@ namespace WindowsMixedReality
 			return false;
 		}
 
-		// Get most recently available tracking data.
-		Frame& TheFrame = GetFrame();
-		CurrentOrientation = TheFrame.HeadOrientation;
-		CurrentPosition = TheFrame.HeadPosition;
+#if WITH_INPUT_SIMULATION
+		if (auto* InputSim = UWindowsMixedRealityInputSimulationEngineSubsystem::GetInputSimulationIfEnabled())
+		{
+			CurrentOrientation = InputSim->GetHeadOrientation();
+			CurrentPosition = InputSim->GetHeadPosition();
+		}
+		else
+#endif
+		{
+			// Get most recently available tracking data.
+			Frame& TheFrame = GetFrame();
+			CurrentOrientation = TheFrame.HeadOrientation;
+			CurrentPosition = TheFrame.HeadPosition;
+		}
 
 		return true;
 	}
@@ -1979,11 +2015,10 @@ namespace WindowsMixedReality
 	HMDInputPressState WindowsMixedReality::FWindowsMixedRealityHMD::GetPressState(HMDHand hand, HMDInputControllerButtons button, bool onlyRegisterClicks)
 	{
 #if WITH_INPUT_SIMULATION
-		auto* InputSim = GEngine->GetEngineSubsystem<UWindowsMixedRealityInputSimulationEngineSubsystem>();
-		if (InputSim && InputSim->IsInputSimulationEnabled())
+		if (auto* InputSim = UWindowsMixedRealityInputSimulationEngineSubsystem::GetInputSimulationIfEnabled())
 		{
 			bool IsPressed;
-			if (InputSim->GetPressState((EControllerHand)hand, (EInputSimulationControllerButtons)button, onlyRegisterClicks, IsPressed))
+			if (InputSim->GetPressState((EControllerHand)hand, (EHMDInputControllerButtons)button, onlyRegisterClicks, IsPressed))
 			{
 				return IsPressed ? HMDInputPressState::Pressed : HMDInputPressState::Released;
 			}
diff --git a/Engine/Plugins/Runtime/WindowsMixedReality/Source/WindowsMixedRealityHMD/Private/WindowsMixedRealityHMD.h b/Engine/Plugins/Runtime/WindowsMixedReality/Source/WindowsMixedRealityHMD/Private/WindowsMixedRealityHMD.h
--- a/Engine/Plugins/Runtime/WindowsMixedReality/Source/WindowsMixedRealityHMD/Private/WindowsMixedRealityHMD.h
+++ b/Engine/Plugins/Runtime/WindowsMixedReality/Source/WindowsMixedRealityHMD/Private/WindowsMixedRealityHMD.h
@@ -100,6 +100,7 @@ namespace WindowsMixedReality
 
 	public:
 		/** IHeadMountedDisplay interface */
+		virtual bool IsHeadTrackingAllowed() const override;
 		virtual bool IsHMDConnected() override;
 		virtual bool IsHMDEnabled() const override;
 		virtual EHMDWornState::Type GetHMDWornState() override;
diff --git a/Engine/Plugins/Runtime/WindowsMixedReality/Source/WindowsMixedRealityHMD/Public/WindowsMixedRealityFunctionLibrary.h b/Engine/Plugins/Runtime/WindowsMixedReality/Source/WindowsMixedRealityHMD/Public/WindowsMixedRealityFunctionLibrary.h
--- a/Engine/Plugins/Runtime/WindowsMixedReality/Source/WindowsMixedRealityHMD/Public/WindowsMixedRealityFunctionLibrary.h
+++ b/Engine/Plugins/Runtime/WindowsMixedReality/Source/WindowsMixedRealityHMD/Public/WindowsMixedRealityFunctionLibrary.h
@@ -26,7 +26,8 @@ enum class EHMDInputControllerButtons : uint8
 	Menu,
 	Thumbstick,
 	Touchpad,
-	TouchpadIsTouched
+	TouchpadIsTouched,
+	Count UMETA(Hidden)
 };
 
 UENUM()
diff --git a/Engine/Plugins/Runtime/WindowsMixedReality/Source/WindowsMixedRealityHandTracking/Private/WindowsMixedRealityHandTracking.cpp b/Engine/Plugins/Runtime/WindowsMixedReality/Source/WindowsMixedRealityHandTracking/Private/WindowsMixedRealityHandTracking.cpp
--- a/Engine/Plugins/Runtime/WindowsMixedReality/Source/WindowsMixedRealityHandTracking/Private/WindowsMixedRealityHandTracking.cpp
+++ b/Engine/Plugins/Runtime/WindowsMixedReality/Source/WindowsMixedRealityHandTracking/Private/WindowsMixedRealityHandTracking.cpp
@@ -356,8 +356,7 @@ bool FWindowsMixedRealityHandTracking::IsHandTrackingStateValid() const
 bool FWindowsMixedRealityHandTracking::GetKeypointTransform(EControllerHand Hand, EWMRHandKeypoint Keypoint, FTransform& OutTransform) const
 {
 #if WITH_INPUT_SIMULATION
-	auto* InputSim = GEngine->GetEngineSubsystem<UWindowsMixedRealityInputSimulationEngineSubsystem>();
-	if (InputSim && InputSim->IsInputSimulationEnabled())
+	if (auto* InputSim = UWindowsMixedRealityInputSimulationEngineSubsystem::GetInputSimulationIfEnabled())
 	{
 		return InputSim->GetHandJointTransform(Hand, Keypoint, OutTransform);
 	}
@@ -373,8 +372,7 @@ bool FWindowsMixedRealityHandTracking::GetKeypointTransform(EControllerHand Hand
 bool FWindowsMixedRealityHandTracking::GetKeypointRadius(EControllerHand Hand, EWMRHandKeypoint Keypoint, float& OutRadius) const
 {
 #if WITH_INPUT_SIMULATION
-	auto* InputSim = GEngine->GetEngineSubsystem<UWindowsMixedRealityInputSimulationEngineSubsystem>();
-	if (InputSim)
+	if (auto* InputSim = UWindowsMixedRealityInputSimulationEngineSubsystem::GetInputSimulationIfEnabled())
 	{
 		return InputSim->GetHandJointRadius(Hand, Keypoint, OutRadius);
 	}
diff --git a/Engine/Plugins/Runtime/WindowsMixedReality/Source/WindowsMixedRealityInputSimulation/Private/WindowsMixedRealityInputSimulationEngineSubsystem.cpp b/Engine/Plugins/Runtime/WindowsMixedReality/Source/WindowsMixedRealityInputSimulation/Private/WindowsMixedRealityInputSimulationEngineSubsystem.cpp
--- a/Engine/Plugins/Runtime/WindowsMixedReality/Source/WindowsMixedRealityInputSimulation/Private/WindowsMixedRealityInputSimulationEngineSubsystem.cpp
+++ b/Engine/Plugins/Runtime/WindowsMixedReality/Source/WindowsMixedRealityInputSimulation/Private/WindowsMixedRealityInputSimulationEngineSubsystem.cpp
@@ -12,7 +12,63 @@
 
 #define LOCTEXT_NAMESPACE "WindowsMixedRealityInputSimulation"
 
-UWindowsMixedRealityInputSimulationEngineSubsystem::HandState* UWindowsMixedRealityInputSimulationEngineSubsystem::GetHandState(EControllerHand Hand)
+UWindowsMixedRealityInputSimulationEngineSubsystem* UWindowsMixedRealityInputSimulationEngineSubsystem::GetInputSimulationIfEnabled()
+{
+	if (!UWindowsMixedRealityInputSimulationEngineSubsystem::IsInputSimulationEnabled())
+	{
+		return nullptr;
+	}
+
+	return GEngine->GetEngineSubsystem<UWindowsMixedRealityInputSimulationEngineSubsystem>();
+}
+
+bool UWindowsMixedRealityInputSimulationEngineSubsystem::IsInputSimulationEnabled()
+{
+#if WITH_EDITOR
+	// Only use input simulation when enabled
+	const UWindowsMixedRealityRuntimeSettings* WMRSettings = UWindowsMixedRealityRuntimeSettings::Get();
+	if (!WMRSettings->bEnableInputSimulation)
+	{
+		return false;
+	}
+
+	// Only if there is a XR system
+	if (!GEngine->XRSystem.IsValid())
+	{
+		return false;
+	}
+
+	// Only if no HMD is connected
+	if (GEngine->XRSystem.Get()->GetHMDDevice()->IsHMDConnected())
+	{
+		return false;
+	}
+
+	return true;
+#else
+	return false;
+#endif
+}
+
+bool UWindowsMixedRealityInputSimulationEngineSubsystem::HasPositionalTracking() const
+{
+	FRWScopeLock ReadLock(DataMutex, SLT_ReadOnly);
+	return bHasPositionalTracking;
+}
+
+const FQuat& UWindowsMixedRealityInputSimulationEngineSubsystem::GetHeadOrientation() const
+{
+	FRWScopeLock ReadLock(DataMutex, SLT_ReadOnly);
+	return HeadOrientation;
+}
+
+const FVector& UWindowsMixedRealityInputSimulationEngineSubsystem::GetHeadPosition() const
+{
+	FRWScopeLock ReadLock(DataMutex, SLT_ReadOnly);
+	return HeadPosition;
+}
+
+FWindowsMixedRealityInputSimulationHandState* UWindowsMixedRealityInputSimulationEngineSubsystem::GetHandState(EControllerHand Hand)
 {
 	if (Hand == EControllerHand::Left)
 	{
@@ -25,17 +81,27 @@ UWindowsMixedRealityInputSimulationEngineSubsystem::HandState* UWindowsMixedReal
 	return nullptr;
 }
 
-const UWindowsMixedRealityInputSimulationEngineSubsystem::HandState* UWindowsMixedRealityInputSimulationEngineSubsystem::GetHandState(EControllerHand Hand) const
+const FWindowsMixedRealityInputSimulationHandState* UWindowsMixedRealityInputSimulationEngineSubsystem::GetHandState(EControllerHand Hand) const
 {
 	return const_cast<UWindowsMixedRealityInputSimulationEngineSubsystem*>(this)->GetHandState(Hand);
 }
 
+ETrackingStatus UWindowsMixedRealityInputSimulationEngineSubsystem::GetControllerTrackingStatus(EControllerHand Hand) const
+{
+	FRWScopeLock ReadLock(DataMutex, SLT_ReadOnly);
+	if (const FWindowsMixedRealityInputSimulationHandState* handState = GetHandState(Hand))
+	{
+		return handState->TrackingStatus;
+	}
+	return ETrackingStatus::NotTracked;
+}
+
 bool UWindowsMixedRealityInputSimulationEngineSubsystem::GetHandJointTransform(EControllerHand Hand, EWMRHandKeypoint Keypoint, FTransform& OutTransform) const
 {
-	FScopeLock Lock(&HandStatesCriticalSection);
-	if (const HandState* handState = GetHandState(Hand))
+	FRWScopeLock ReadLock(DataMutex, SLT_ReadOnly);
+	if (const FWindowsMixedRealityInputSimulationHandState* handState = GetHandState(Hand))
 	{
-		if (handState->IsValid)
+		if (handState->bHasJointPoses)
 		{
 			OutTransform = handState->KeypointTransforms[(int32)Keypoint];
 			return true;
@@ -46,10 +112,10 @@ bool UWindowsMixedRealityInputSimulationEngineSubsystem::GetHandJointTransform(E
 
 bool UWindowsMixedRealityInputSimulationEngineSubsystem::GetHandJointRadius(EControllerHand Hand, EWMRHandKeypoint Keypoint, float& OutRadius) const
 {
-	FScopeLock Lock(&HandStatesCriticalSection);
-	if (const HandState* handState = GetHandState(Hand))
+	FRWScopeLock ReadLock(DataMutex, SLT_ReadOnly);
+	if (const FWindowsMixedRealityInputSimulationHandState* handState = GetHandState(Hand))
 	{
-		if (handState->IsValid)
+		if (handState->bHasJointPoses)
 		{
 			OutRadius = handState->KeypointRadii[(int32)Keypoint];
 			return true;
@@ -58,102 +124,65 @@ bool UWindowsMixedRealityInputSimulationEngineSubsystem::GetHandJointRadius(ECon
 	return false;
 }
 
-void UWindowsMixedRealityInputSimulationEngineSubsystem::SetHandJointOrientationAndPosition(EControllerHand Hand, EWMRHandKeypoint Keypoint, const FRotator& Orientation, const FVector& Position)
-{
-	FScopeLock Lock(&HandStatesCriticalSection);
-	if (HandState* handState = GetHandState(Hand))
-	{
-		handState->KeypointTransforms[(int32)Keypoint].SetComponents(Orientation.Quaternion(), Position, FVector::OneVector);
-	}
-}
-
-void UWindowsMixedRealityInputSimulationEngineSubsystem::SetHandJointRadius(EControllerHand Hand, EWMRHandKeypoint Keypoint, float Radius)
+bool UWindowsMixedRealityInputSimulationEngineSubsystem::HasJointPoses(EControllerHand Hand) const
 {
-	FScopeLock Lock(&HandStatesCriticalSection);
-	if (HandState* handState = GetHandState(Hand))
+	FRWScopeLock ReadLock(DataMutex, SLT_ReadOnly);
+	if (const FWindowsMixedRealityInputSimulationHandState* handState = GetHandState(Hand))
 	{
-		handState->KeypointRadii[(int32)Keypoint] = Radius;
-	}
-}
-
-bool UWindowsMixedRealityInputSimulationEngineSubsystem::IsHandStateValid(EControllerHand Hand) const
-{
-	FScopeLock Lock(&HandStatesCriticalSection);
-	if (const HandState* handState = GetHandState(Hand))
-	{
-		return handState->IsValid;
+		return handState->bHasJointPoses;
 	}
 	return false;
 }
 
-void UWindowsMixedRealityInputSimulationEngineSubsystem::SetHandStateValid(EControllerHand Hand, bool IsValid)
+bool UWindowsMixedRealityInputSimulationEngineSubsystem::GetPressState(EControllerHand Hand, EHMDInputControllerButtons Button, bool OnlyRegisterClicks, bool& OutPressState) const
 {
-	FScopeLock Lock(&HandStatesCriticalSection);
-	if (HandState* handState = GetHandState(Hand))
+	FRWScopeLock ReadLock(DataMutex, SLT_ReadOnly);
+	if (const FWindowsMixedRealityInputSimulationHandState* handState = GetHandState(Hand))
 	{
-		handState->IsValid = IsValid;
-	}
-}
+		bool IsPressed = handState->IsButtonPressed[(int32)Button];
+		bool WasPressed = handState->PrevButtonPressed[(int32)Button];
 
-void UWindowsMixedRealityInputSimulationEngineSubsystem::SetPressState(EControllerHand Hand, EInputSimulationControllerButtons Button, bool IsPressed)
-{
-	FScopeLock Lock(&HandStatesCriticalSection);
-	if (HandState* handState = GetHandState(Hand))
-	{
-		handState->PrevButtonPressed[(int32)Button] = handState->IsButtonPressed[(int32)Button];
-		handState->IsButtonPressed[(int32)Button] = IsPressed;
-	}
-}
+		OutPressState = IsPressed;
 
-bool UWindowsMixedRealityInputSimulationEngineSubsystem::GetPressState(EControllerHand Hand, EInputSimulationControllerButtons Button, bool OnlyRegisterClicks, bool& OutPressState) const
-{
-	FScopeLock Lock(&HandStatesCriticalSection);
-	if (const HandState* handState = GetHandState(Hand))
-	{
-		if (handState->IsValid)
+		if (OnlyRegisterClicks)
+		{
+			return IsPressed != WasPressed;
+		}
+		else
 		{
-			bool IsPressed = handState->IsButtonPressed[(int32)Button];
-			bool WasPressed = handState->PrevButtonPressed[(int32)Button];
-			if (OnlyRegisterClicks)
-			{
-				OutPressState = IsPressed != WasPressed ? IsPressed : false;
-			}
-			else
-			{
-				OutPressState = IsPressed;
-			}
 			return true;
 		}
 	}
 	return false;
 }
 
-bool UWindowsMixedRealityInputSimulationEngineSubsystem::IsInputSimulationEnabled() const
+void UWindowsMixedRealityInputSimulationEngineSubsystem::UpdateSimulatedData(
+	bool HasTracking,
+	const FQuat& NewHeadOrientation,
+	const FVector& NewHeadPosition,
+	const FWindowsMixedRealityInputSimulationHandState& LeftHandState,
+	const FWindowsMixedRealityInputSimulationHandState& RightHandState)
 {
-#if WITH_EDITOR
-	// Only use input simulation when enabled
-	const UWindowsMixedRealityRuntimeSettings* WMRSettings = UWindowsMixedRealityRuntimeSettings::Get();
-	if (!WMRSettings->bEnableInputSimulation)
-	{
-		return false;
-	}
+	FRWScopeLock WriteLock(DataMutex, SLT_Write);
 
-	// Only if there is a XR system
-	if (!GEngine->XRSystem.IsValid())
-	{
-		return false;
-	}
+	bHasPositionalTracking = HasTracking;
+	HeadOrientation = NewHeadOrientation;
+	HeadPosition = NewHeadPosition;
 
-	// Only if no HMD is connected
-	if (GEngine->XRSystem.Get()->GetHMDDevice()->IsHMDConnected())
-	{
-		return false;
-	}
+	UpdateSimulatedHandState(EControllerHand::Left, LeftHandState);
+	UpdateSimulatedHandState(EControllerHand::Right, RightHandState);
+}
 
-	return true;
-#else
-	return false;
-#endif
+void UWindowsMixedRealityInputSimulationEngineSubsystem::UpdateSimulatedHandState(
+	EControllerHand Hand,
+	const FWindowsMixedRealityInputSimulationHandState& NewHandState)
+{
+	FWindowsMixedRealityInputSimulationHandState* HandStatePtr = GetHandState(Hand);
+	check(HandStatePtr);
+
+	FWindowsMixedRealityInputSimulationHandState::ButtonStateArray WasButtonPressed(HandStatePtr->IsButtonPressed);
+	*HandStatePtr = NewHandState;
+	HandStatePtr->PrevButtonPressed = WasButtonPressed;
 }
 
 bool UWindowsMixedRealityInputSimulationEngineSubsystem::ShouldCreateSubsystem(UObject* Outer) const
diff --git a/Engine/Plugins/Runtime/WindowsMixedReality/Source/WindowsMixedRealityInputSimulation/Public/WindowsMixedRealityInputSimulationEngineSubsystem.h b/Engine/Plugins/Runtime/WindowsMixedReality/Source/WindowsMixedRealityInputSimulation/Public/WindowsMixedRealityInputSimulationEngineSubsystem.h
--- a/Engine/Plugins/Runtime/WindowsMixedReality/Source/WindowsMixedRealityInputSimulation/Public/WindowsMixedRealityInputSimulationEngineSubsystem.h
+++ b/Engine/Plugins/Runtime/WindowsMixedReality/Source/WindowsMixedRealityInputSimulation/Public/WindowsMixedRealityInputSimulationEngineSubsystem.h
@@ -4,23 +4,31 @@
 
 #include "CoreMinimal.h"
 #include "InputCoreTypes.h"
+#include "IMotionController.h"
 #include "Subsystems/EngineSubsystem.h"
+#include "WindowsMixedRealityFunctionLibrary.h"
 #include "WindowsMixedRealityHandTrackingTypes.h"
+#include "Containers/StaticBitArray.h"
+#include "Containers/StaticArray.h"
 
 #include "WindowsMixedRealityInputSimulationEngineSubsystem.generated.h"
 
-UENUM()
-enum class EInputSimulationControllerButtons : uint8
+/** Simulated data for one hand. */
+struct FWindowsMixedRealityInputSimulationHandState
 {
-	Select,
-	Grasp,
-	Menu,
-	Thumbstick,
-	Touchpad,
-	TouchpadIsTouched
-};
+	typedef TStaticArray<FTransform, (uint32)EWMRHandKeypointCount> KeypointTransformArray;
+	typedef TStaticArray<float, (uint32)EWMRHandKeypointCount> KeypointRadiusArray;
+	typedef TStaticBitArray<uint32(EHMDInputControllerButtons::Count)> ButtonStateArray;
+
+	ETrackingStatus TrackingStatus = ETrackingStatus::NotTracked;
 
-const int32 EInputSimulationControllerButtonsCount = static_cast<int32>(EInputSimulationControllerButtons::TouchpadIsTouched) + 1;
+	bool bHasJointPoses = false;
+	KeypointTransformArray KeypointTransforms = KeypointTransformArray(FTransform::Identity);
+	KeypointRadiusArray KeypointRadii = KeypointRadiusArray(1.0f);
+
+	ButtonStateArray IsButtonPressed = ButtonStateArray(false);
+	ButtonStateArray PrevButtonPressed = ButtonStateArray(false);
+};
 
 /** Engine subsystem that stores input simulation data for access by the XR device. */
 UCLASS(ClassGroup = WindowsMixedReality)
@@ -31,33 +39,52 @@ class WINDOWSMIXEDREALITYINPUTSIMULATION_API UWindowsMixedRealityInputSimulation
 
 public:
 
-	struct HandState
-	{
-		bool IsValid = false;
+	/** Checks both the user settings and HMD availability to determine if input simulation is enabled.
+	 * Note: The engine subsystem is created before the XRSystem, so it may exist even if a HMD is connected
+	 *       and input simulation is ultimately not used at runtime.
+	 */
+	static bool IsInputSimulationEnabled();
+
+	/** Utility function to ensure input simulation is only used when enabled.
+	 * Returns the engine subsystem if it exists and if input simulation is enabled.
+	 */
+	static UWindowsMixedRealityInputSimulationEngineSubsystem* GetInputSimulationIfEnabled();
 		
-		FTransform KeypointTransforms[EWMRHandKeypointCount];
-		float KeypointRadii[EWMRHandKeypointCount];
+	/** True if the HMD has a valid head position vector. */
+	bool HasPositionalTracking() const;
+	/** Orientation of the HMD. */
+	const FQuat& GetHeadOrientation() const;
+	/** Position of the HMD if positional tracking is available. */
+	const FVector& GetHeadPosition() const;
 		
-		bool IsButtonPressed[EInputSimulationControllerButtonsCount] = { false };
-		bool PrevButtonPressed[EInputSimulationControllerButtonsCount] = { false };
-	};
+	/** Current tracking status of a hand. */
+	ETrackingStatus GetControllerTrackingStatus(EControllerHand Hand) const;
 
+	/** Get the transform of a hand joint.
+	 *  Returns true if the output transform is valid.
+	 */
 	bool GetHandJointTransform(EControllerHand Hand, EWMRHandKeypoint Keypoint, FTransform& OutTransform) const;
+	/** Get the radius of a hand joint.
+	 *  Returns true if the output radius is valid.
+	 */
 	bool GetHandJointRadius(EControllerHand Hand, EWMRHandKeypoint Keypoint, float& OutRadius) const;
-	void SetHandJointOrientationAndPosition(EControllerHand Hand, EWMRHandKeypoint Keypoint, const FRotator& Orientation, const FVector& Position);
-	void SetHandJointRadius(EControllerHand Hand, EWMRHandKeypoint Keypoint, float Radius);
 	
-	bool IsHandStateValid(EControllerHand Hand) const;
-	void SetHandStateValid(EControllerHand Hand, bool IsValid);
-
-	void SetPressState(EControllerHand Hand, EInputSimulationControllerButtons Button, bool IsPressed);
-	bool GetPressState(EControllerHand Hand, EInputSimulationControllerButtons Button, bool OnlyRegisterClicks, bool& OutPressState) const;
+	/** True if joint transforms are available for a hand. */
+	bool HasJointPoses(EControllerHand Hand) const;
 
-	/** Checks both the user settings and HMD availability to determine if input simulation is enabled.
-	 * Note: The engine subsystem is created before the XRSystem, so it may exist even if a HMD is connected
-	 *       and input simulation is ultimately not used at runtime.
+	/** State of buttons and gestures of the hand controller.
+	 * If OnlyRegisterClicks is true only changes from unpressed to pressed will be registered.
+	 * Returns true only if the button state is valid.
 	 */
-	bool IsInputSimulationEnabled() const;
+	bool GetPressState(EControllerHand Hand, EHMDInputControllerButtons Button, bool OnlyRegisterClicks, bool& OutPressState) const;
+
+	/** Update the simulated data. */
+	void UpdateSimulatedData(
+		bool HasTracking,
+		const FQuat& NewHeadOrientation,
+		const FVector& NewHeadPosition,
+		const FWindowsMixedRealityInputSimulationHandState& LeftHandState,
+		const FWindowsMixedRealityInputSimulationHandState& RightHandState);
 
 	//
 	// USubsystem implementation
@@ -68,14 +95,23 @@ public:
 
 private:
 
-	HandState* GetHandState(EControllerHand Hand);
-	const HandState* GetHandState(EControllerHand Hand) const;
+	/** Utility function for getting the state of a hand. */
+	FWindowsMixedRealityInputSimulationHandState* GetHandState(EControllerHand Hand);
+	/** Utility function for getting the state of a hand. */
+	const FWindowsMixedRealityInputSimulationHandState* GetHandState(EControllerHand Hand) const;
+
+	/** Utility function for updating the state of a hand. */
+	void UpdateSimulatedHandState(EControllerHand Hand, const FWindowsMixedRealityInputSimulationHandState& NewHandState);
 
 private:
 
-	HandState HandStates[2];
+	bool bHasPositionalTracking = false;
+	FQuat HeadOrientation = FQuat::Identity;
+	FVector HeadPosition = FVector::ZeroVector;
+
+	FWindowsMixedRealityInputSimulationHandState HandStates[2];
 
-	/** Critical section for reading/writing hand states. */
-	mutable FCriticalSection HandStatesCriticalSection;
+	/** Lock for accessing simulated data */
+	mutable FRWLock DataMutex;
 
 };
diff --git a/Engine/Plugins/Runtime/WindowsMixedReality/Source/WindowsMixedRealityInputSimulation/WindowsMixedRealityInputSimulation.Build.cs b/Engine/Plugins/Runtime/WindowsMixedReality/Source/WindowsMixedRealityInputSimulation/WindowsMixedRealityInputSimulation.Build.cs
--- a/Engine/Plugins/Runtime/WindowsMixedReality/Source/WindowsMixedRealityInputSimulation/WindowsMixedRealityInputSimulation.Build.cs
+++ b/Engine/Plugins/Runtime/WindowsMixedReality/Source/WindowsMixedRealityInputSimulation/WindowsMixedRealityInputSimulation.Build.cs
@@ -9,9 +9,18 @@ namespace UnrealBuildTool.Rules
 		public WindowsMixedRealityInputSimulation(ReadOnlyTargetRules Target)
 				: base(Target)
 		{
+			PublicIncludePathModuleNames.AddRange(
+				new string[]
+				{
+					"WindowsMixedRealityHandTracking"
+				}
+			);
+
             PublicDependencyModuleNames.AddRange(
 				new string[]
 				{
+					"HeadMountedDisplay",
+					"WindowsMixedRealityHMD"
 				}
 			);
 
@@ -21,8 +30,6 @@ namespace UnrealBuildTool.Rules
 					"Core",
 					"CoreUObject",
 					"Engine",
-					"HeadMountedDisplay",
-					"WindowsMixedRealityHandTracking",
 					"WindowsMixedRealityRuntimeSettings",
 				}
 			);
diff --git a/Engine/Plugins/Runtime/WindowsMixedReality/Source/WindowsMixedRealitySpatialInput/Private/WindowsMixedRealitySpatialInput.cpp b/Engine/Plugins/Runtime/WindowsMixedReality/Source/WindowsMixedRealitySpatialInput/Private/WindowsMixedRealitySpatialInput.cpp
--- a/Engine/Plugins/Runtime/WindowsMixedReality/Source/WindowsMixedRealitySpatialInput/Private/WindowsMixedRealitySpatialInput.cpp
+++ b/Engine/Plugins/Runtime/WindowsMixedReality/Source/WindowsMixedRealitySpatialInput/Private/WindowsMixedRealitySpatialInput.cpp
@@ -9,6 +9,11 @@
 #include "WindowsMixedRealitySpatialInputTypes.h"
 #include "WindowsMixedRealityAvailability.h"
 
+#if WITH_INPUT_SIMULATION
+#include "Engine/Engine.h"
+#include "WindowsMixedRealityInputSimulationEngineSubsystem.h"
+#endif
+
 #define LOCTEXT_NAMESPACE "WindowsMixedRealitySpatialInput"
 #define MotionControllerDeviceTypeName "WindowsMixedRealitySpatialInput"
 
@@ -800,6 +805,12 @@ namespace WindowsMixedReality
 	ETrackingStatus FWindowsMixedRealitySpatialInput::GetControllerTrackingStatus(const int32 ControllerIndex, const EControllerHand DeviceHand) const
 	{
 #if WITH_WINDOWS_MIXED_REALITY
+#if WITH_INPUT_SIMULATION
+		if (auto* InputSim = UWindowsMixedRealityInputSimulationEngineSubsystem::GetInputSimulationIfEnabled())
+		{
+			return InputSim->GetControllerTrackingStatus(DeviceHand);
+		}
+#endif
 		HMDHand hand = (HMDHand)((int)DeviceHand);
 		HMDTrackingStatus trackingStatus = FWindowsMixedRealityStatics::GetControllerTrackingStatus(hand);
 
diff --git a/Engine/Plugins/Runtime/WindowsMixedReality/Source/WindowsMixedRealitySpatialInput/WindowsMixedRealitySpatialInput.Build.cs b/Engine/Plugins/Runtime/WindowsMixedReality/Source/WindowsMixedRealitySpatialInput/WindowsMixedRealitySpatialInput.Build.cs
--- a/Engine/Plugins/Runtime/WindowsMixedReality/Source/WindowsMixedRealitySpatialInput/WindowsMixedRealitySpatialInput.Build.cs
+++ b/Engine/Plugins/Runtime/WindowsMixedReality/Source/WindowsMixedRealitySpatialInput/WindowsMixedRealitySpatialInput.Build.cs
@@ -68,6 +68,16 @@ namespace UnrealBuildTool.Rules
                 PCHUsage = PCHUsageMode.NoSharedPCHs;
                 PrivatePCHHeaderFile = "Private/WindowsMixedRealitySpatialInput.h";
             }
-        }
-    }
+
+			if (Target.Platform == UnrealTargetPlatform.Win64 && Target.bBuildEditor == true)
+			{
+				PrivateDependencyModuleNames.Add("WindowsMixedRealityInputSimulation");
+				PrivateDefinitions.Add("WITH_INPUT_SIMULATION=1");
+			}
+			else
+			{
+				PrivateDefinitions.Add("WITH_INPUT_SIMULATION=0");
+			}
+		}
+	}
 }
\ No newline at end of file
