Unshelved from pending changelist '11191641':

PR #6642: Rotate hand joint transforms to match UE convention (Contributed by [REDACTED_PERSON])

diff --git a/Engine/Plugins/Runtime/WindowsMixedReality/Source/WindowsMixedRealityHandTracking/Private/WindowsMixedRealityHandTracking.cpp b/Engine/Plugins/Runtime/WindowsMixedReality/Source/WindowsMixedRealityHandTracking/Private/WindowsMixedRealityHandTracking.cpp
--- a/Engine/Plugins/Runtime/WindowsMixedReality/Source/WindowsMixedRealityHandTracking/Private/WindowsMixedRealityHandTracking.cpp
+++ b/Engine/Plugins/Runtime/WindowsMixedReality/Source/WindowsMixedRealityHandTracking/Private/WindowsMixedRealityHandTracking.cpp
@@ -355,18 +355,27 @@ bool FWindowsMixedRealityHandTracking::IsHandTrackingStateValid() const
 
 bool FWindowsMixedRealityHandTracking::GetKeypointTransform(EControllerHand Hand, EWMRHandKeypoint Keypoint, FTransform& OutTransform) const
 {
+	bool gotTransform = false;
+ 
 #if WITH_INPUT_SIMULATION
 	if (auto* InputSim = UWindowsMixedRealityInputSimulationEngineSubsystem::GetInputSimulationIfEnabled())
 	{
-		return InputSim->GetHandJointTransform(Hand, Keypoint, OutTransform);
+		gotTransform = InputSim->GetHandJointTransform(Hand, Keypoint, OutTransform);
 	}
 	else
 #endif
 	{
 		const FWindowsMixedRealityHandTracking::FHandState& HandState = (Hand == EControllerHand::Left) ? GetLeftHandState() : GetRightHandState();
-
-		return HandState.GetTransform(Keypoint, OutTransform);
+		gotTransform = HandState.GetTransform(Keypoint, OutTransform);
+		// Rotate to match UE space conventions (positive-x forward, positive-y right, positive-z up)
+		OutTransform.SetRotation(OutTransform.GetRotation() * FQuat(FVector::RightVector, PI));
+	}
+	if (gotTransform)
+	{
+		// Convert to UE world space
+		OutTransform *= UHeadMountedDisplayFunctionLibrary::GetTrackingToWorldTransform(GWorld);
 	}
+	return gotTransform;
 }
 
 bool FWindowsMixedRealityHandTracking::GetKeypointRadius(EControllerHand Hand, EWMRHandKeypoint Keypoint, float& OutRadius) const
