Fix build break from merged Microsoft submission for Hololens hand-tracking

diff --git a/Engine/Plugins/Runtime/WindowsMixedReality/Source/WindowsMixedRealityHandTracking/Private/WindowsMixedRealityHandTracking.cpp b/Engine/Plugins/Runtime/WindowsMixedReality/Source/WindowsMixedRealityHandTracking/Private/WindowsMixedRealityHandTracking.cpp
--- a/Engine/Plugins/Runtime/WindowsMixedReality/Source/WindowsMixedRealityHandTracking/Private/WindowsMixedRealityHandTracking.cpp
+++ b/Engine/Plugins/Runtime/WindowsMixedReality/Source/WindowsMixedRealityHandTracking/Private/WindowsMixedRealityHandTracking.cpp
@@ -12,6 +12,7 @@
 #include "Features/IModularFeatures.h"
 #include "IWindowsMixedRealityHandTrackingPlugin.h"
 #include "ILiveLinkClient.h"
+#include "HeadMountedDisplayFunctionLibrary.h"
 
 #if WITH_INPUT_SIMULATION
 #include "WindowsMixedRealityInputSimulationEngineSubsystem.h"
