Fix CIS health issue

diff --git a/Engine/Plugins/Runtime/WindowsMixedReality/Source/WindowsMixedRealityInputSimulation/Public/WindowsMixedRealityInputSimulationEngineSubsystem.h b/Engine/Plugins/Runtime/WindowsMixedReality/Source/WindowsMixedRealityInputSimulation/Public/WindowsMixedRealityInputSimulationEngineSubsystem.h
--- a/Engine/Plugins/Runtime/WindowsMixedReality/Source/WindowsMixedRealityInputSimulation/Public/WindowsMixedRealityInputSimulationEngineSubsystem.h
+++ b/Engine/Plugins/Runtime/WindowsMixedReality/Source/WindowsMixedRealityInputSimulation/Public/WindowsMixedRealityInputSimulationEngineSubsystem.h
@@ -26,8 +26,8 @@ struct FWindowsMixedRealityInputSimulationHandState
 	KeypointTransformArray KeypointTransforms = KeypointTransformArray(FTransform::Identity);
 	KeypointRadiusArray KeypointRadii = KeypointRadiusArray(1.0f);
 
-	ButtonStateArray IsButtonPressed = ButtonStateArray(false);
-	ButtonStateArray PrevButtonPressed = ButtonStateArray(false);
+	ButtonStateArray IsButtonPressed = ButtonStateArray();
+	ButtonStateArray PrevButtonPressed = ButtonStateArray();
 };
 
 /** Engine subsystem that stores input simulation data for access by the XR device. */
