From Microsoft: Hololens/WMR Support Windows Server 2019

diff --git a/Engine/Plugins/Runtime/WindowsMixedReality/Source/WindowsMixedRealityHMD/Private/WindowsMixedRealityInteropLoader.cpp b/Engine/Plugins/Runtime/WindowsMixedReality/Source/WindowsMixedRealityHMD/Private/WindowsMixedRealityInteropLoader.cpp
--- a/Engine/Plugins/Runtime/WindowsMixedReality/Source/WindowsMixedRealityHMD/Private/WindowsMixedRealityInteropLoader.cpp
+++ b/Engine/Plugins/Runtime/WindowsMixedReality/Source/WindowsMixedRealityHMD/Private/WindowsMixedRealityInteropLoader.cpp
@@ -25,7 +25,8 @@ namespace WindowsMixedReality
 		FPlatformMisc::GetOSVersions(OSVersionLabel, OSSubVersionLabel);
 		// GetOSVersion returns the Win10 release version in the OSVersion rather than the OSSubVersion, so parse it out ourselves
 		OSSubVersionLabel = OSVersionLabel;
-		bool bHasSupportedWindowsVersion = OSSubVersionLabel.RemoveFromStart("Windows 10 (Release ") && OSSubVersionLabel.RemoveFromEnd(")") && (FCString::Atoi(*OSSubVersionLabel) >= MIN_WIN_10_VERSION_FOR_WMR);
+		bool bHasSupportedWindowsVersion = OSSubVersionLabel.StartsWith("Windows 10") || OSSubVersionLabel.StartsWith("Windows Server 2019");
+		bool bHasSupportedWindowsBuild = false;
 
 		// If we can't find Win10 version, check for Windows Server equivalent
 		if (!bHasSupportedWindowsVersion)
@@ -35,6 +36,19 @@ namespace WindowsMixedReality
 		}
 
 		if (bHasSupportedWindowsVersion)
+		{
+			FString BeginMarker = TEXT("(Release ");
+			int Begin = OSSubVersionLabel.Find(BeginMarker);
+			int End = OSSubVersionLabel.Find(")", ESearchCase::IgnoreCase, ESearchDir::FromStart, Begin);
+			if (Begin >= 0 && End >= 0)
+			{
+				Begin += BeginMarker.Len();
+				OSSubVersionLabel = OSSubVersionLabel.Mid(Begin, End - Begin);
+				bHasSupportedWindowsBuild = FCString::Atoi(*OSSubVersionLabel) >= MIN_WIN_10_VERSION_FOR_WMR;
+			}
+		}
+
+		if (bHasSupportedWindowsBuild && bHasSupportedWindowsVersion)
 		{
 			// Get the base directory of this plugin
 			FString BaseDir = IPluginManager::Get().FindPlugin("WindowsMixedReality")->GetBaseDir();
diff --git a/Engine/Source/Runtime/Core/Private/Windows/WindowsPlatformMisc.cpp b/Engine/Source/Runtime/Core/Private/Windows/WindowsPlatformMisc.cpp
--- a/Engine/Source/Runtime/Core/Private/Windows/WindowsPlatformMisc.cpp
+++ b/Engine/Source/Runtime/Core/Private/Windows/WindowsPlatformMisc.cpp
@@ -249,7 +249,7 @@ int32 GetOSVersionsHelper( TCHAR* OutOSVersionLabel, int32 OSVersionLabelLength,
 				}
 				else
 				{
-					OSVersionLabel = TEXT("Windows Server Technical Preview");
+					OSVersionLabel = TEXT("Windows Server 2019");
 				}
 
 				// For Windows 10, get the release number and append that to the string too (eg. 1709 = Fall Creators Update). There doesn't seem to be any good way to get
