Support instanced mobile multi-view for Hololens

diff --git a/Engine/Plugins/Runtime/WindowsMixedReality/Source/WindowsMixedRealityHMD/Private/WindowsMixedRealityCustomPresent.h b/Engine/Plugins/Runtime/WindowsMixedReality/Source/WindowsMixedRealityHMD/Private/WindowsMixedRealityCustomPresent.h
--- a/Engine/Plugins/Runtime/WindowsMixedReality/Source/WindowsMixedRealityHMD/Private/WindowsMixedRealityCustomPresent.h
+++ b/Engine/Plugins/Runtime/WindowsMixedReality/Source/WindowsMixedRealityHMD/Private/WindowsMixedRealityCustomPresent.h
@@ -29,6 +29,11 @@ namespace WindowsMixedReality
 		{
 			// Get the D3D11 context.
 			device->GetImmediateContext(&D3D11Context);
+
+#if PLATFORM_HOLOLENS
+			static const auto CVarMobileMultiViewDirect = IConsoleManager::Get().FindTConsoleVariableDataInt(TEXT("vr.MobileMultiView.Direct"));
+			bIsMobileMultiViewDirectEnabled = (CVarMobileMultiViewDirect && CVarMobileMultiViewDirect->GetValueOnAnyThread() != 0);
+#endif
 		}
 
 		// Inherited via FRHICustomPresent
@@ -47,7 +52,10 @@ namespace WindowsMixedReality
 				return false;
 			}
 
-			hmd->CopyResources(D3D11Context, ViewportTexture);
+			if (!bIsMobileMultiViewDirectEnabled)
+			{
+				hmd->CopyResources(D3D11Context, ViewportTexture);
+			}
 			
 			return hmd->Present();
 #else
@@ -86,5 +94,6 @@ namespace WindowsMixedReality
 
 		ID3D11DeviceContext* D3D11Context = nullptr;
 		ID3D11Texture2D* ViewportTexture = nullptr;
+		bool bIsMobileMultiViewDirectEnabled = false;
 	};
 }
\ No newline at end of file
diff --git a/Engine/Plugins/Runtime/WindowsMixedReality/Source/WindowsMixedRealityHMD/Private/WindowsMixedRealityHMD.cpp b/Engine/Plugins/Runtime/WindowsMixedReality/Source/WindowsMixedRealityHMD/Private/WindowsMixedRealityHMD.cpp
--- a/Engine/Plugins/Runtime/WindowsMixedReality/Source/WindowsMixedRealityHMD/Private/WindowsMixedRealityHMD.cpp
+++ b/Engine/Plugins/Runtime/WindowsMixedReality/Source/WindowsMixedRealityHMD/Private/WindowsMixedRealityHMD.cpp
@@ -625,6 +625,18 @@ namespace WindowsMixedReality
 			// Because Windows Mixed Reality can only have 2 rendering parameters in flight at any time, this is fatal.
 			this->bRequestRestart = true;
 		}
+
+#if PLATFORM_HOLOLENS
+		static const auto CVarMobileMultiViewDirect = IConsoleManager::Get().FindTConsoleVariableDataInt(TEXT("vr.MobileMultiView.Direct"));
+		const bool bIsMobileMultiViewDirectEnabled = (CVarMobileMultiViewDirect && CVarMobileMultiViewDirect->GetValueOnRenderThread() != 0);
+		if (bIsMobileMultiViewDirectEnabled)
+		{
+			FD3D11DynamicRHI* DynamicRHI = static_cast<FD3D11DynamicRHI*>(GDynamicRHI);
+			ID3D11Texture2D* Texture = HMD->GetBackBufferTexture();
+			FTexture2DArrayRHIRef BackBuffer = DynamicRHI->RHICreateTexture2DArrayFromResource(PF_B8G8R8A8, TexCreate_RenderTargetable | TexCreate_ShaderResource, FClearValueBinding::None, Texture);
+			GDynamicRHI->RHIAliasTextureResources((FTextureRHIRef&)CurrentBackBuffer, (FTextureRHIRef&)BackBuffer);
+		}
+#endif
 #endif
 	}
 
@@ -641,6 +653,13 @@ namespace WindowsMixedReality
 		{
 			static FVector2D SrcNormRectMin(0.05f, 0.2f);
 			static FVector2D SrcNormRectMax(0.45f, 0.8f);
+
+#if PLATFORM_HOLOLENS
+			static const auto CVarMobileMultiView = IConsoleManager::Get().FindTConsoleVariableDataInt(TEXT("vr.MobileMultiView"));
+			const bool bIsMobileMultiViewEnabled = (CVarMobileMultiView && CVarMobileMultiView->GetValueOnRenderThread() != 0);
+			SrcNormRectMax.X = bIsMobileMultiViewEnabled ? 0.95f : 0.45f;
+#endif
+
 			return FIntRect(EyeTexture->GetSizeX() * SrcNormRectMin.X, EyeTexture->GetSizeY() * SrcNormRectMin.Y, EyeTexture->GetSizeX() * SrcNormRectMax.X, EyeTexture->GetSizeY() * SrcNormRectMax.Y);
 		}
 	}
@@ -1016,7 +1035,11 @@ namespace WindowsMixedReality
 					if (HMD->GetDisplayDimensions(Width, Height))
 					{
 						SceneVP->SetViewportSize(
+#if PLATFORM_HOLOLENS
+							Width,
+#else
 							Width * 2,
+#endif
 							Height);
 
 						Window->SetViewportSizeDrivenByWindow(false);
@@ -1235,7 +1258,15 @@ namespace WindowsMixedReality
 		tcHeight = 0;
 #endif
 
-		int RenderTargetWidth = Width * 2 + tcWidth;
+#if PLATFORM_HOLOLENS
+		static const auto CVarMobileMultiViewDirect = IConsoleManager::Get().FindTConsoleVariableDataInt(TEXT("vr.MobileMultiView.Direct"));
+		const bool bIsMobileMultiViewDirectEnabled = (CVarMobileMultiViewDirect && CVarMobileMultiViewDirect->GetValueOnAnyThread() != 0);
+		int Offset = bIsMobileMultiViewDirectEnabled ? 0 : Width;
+#else
+		int Offset = Width;
+#endif
+
+		int RenderTargetWidth = Width + Offset + tcWidth;
 		int RenderTargetHeight = FMath::Max(Height, tcHeight);
 
 		// We always prefer the nearest multiple of 4 for our buffer sizes. Make sure we round up here,
@@ -1245,14 +1276,19 @@ namespace WindowsMixedReality
 
 		float eyeUVX = (float)Width / (float)Size.X;
 		float eyeUVY = (float)Height / (float)Size.Y;
+#if PLATFORM_HOLOLENS
+		float offUVX = bIsMobileMultiViewDirectEnabled ? 0.0f : eyeUVX;
+#else
+		float offUVX = eyeUVX;
+#endif
 
 		float camUVX = (float)tcWidth / (float)Size.X;
 		float camUVY = (float)tcHeight / (float)Size.Y;
 
-		//                                            width,    height,    x,          y,  uvOffsetX,   uvOffsetY, uvScaleX, uvScaleY
-		wmrRenderTargets[0] = RenderTargetDescription(Width,    Height,    0,          0,  0.0f,        0.0f,      eyeUVX,   eyeUVY);  // Left Eye
-		wmrRenderTargets[1] = RenderTargetDescription(Width,    Height,    Width,      0,  eyeUVX,      0.0f,      eyeUVX,   eyeUVY);  // Right Eye
-		wmrRenderTargets[2] = RenderTargetDescription(tcWidth,  tcHeight,  Width * 2,  0,  eyeUVX * 2,  0.0f,      camUVX,   camUVY);  // Third Camera
+		//                                            width,    height,    x,           y,  uvOffsetX,   uvOffsetY, uvScaleX, uvScaleY
+		wmrRenderTargets[0] = RenderTargetDescription(Width,    Height,    0,           0,  0.0f,        0.0f,      eyeUVX,   eyeUVY);  // Left Eye
+		wmrRenderTargets[1] = RenderTargetDescription(Width,    Height,    Offset,      0,  offUVX,      0.0f,      eyeUVX,   eyeUVY);  // Right Eye
+		wmrRenderTargets[2] = RenderTargetDescription(tcWidth,  tcHeight,  Offset * 2,  0,  offUVX * 2,  0.0f,      camUVX,   camUVY);  // Third Camera
 
 		return Size;
 	}
@@ -1290,11 +1326,24 @@ namespace WindowsMixedReality
 			SpectatorScreenController->RenderSpectatorScreen_RenderThread(RHICmdList, BackBuffer, SrcTexture, WindowSize);
 		}
 
-		// We keep refs to the depth texture in the hmd, so this function is non-const.
-		// RenderTexture_RenderThread should perhaps be refactored or made non-const.  
-		// But to get this into a hotfix of 4.23 we shall simply cast.
-		FWindowsMixedRealityHMD* nonconstthis = const_cast<FWindowsMixedRealityHMD*>(this);
-		nonconstthis->CreateHMDDepthTexture(RHICmdList);
+#if PLATFORM_HOLOLENS
+		static const auto CVarMobileMultiView = IConsoleManager::Get().FindTConsoleVariableDataInt(TEXT("vr.MobileMultiView"));
+		const bool bIsMobileMultiViewEnabled = (CVarMobileMultiView && CVarMobileMultiView->GetValueOnRenderThread() != 0);
+
+		if (bIsMobileMultiViewEnabled)
+		{
+			ID3D11Texture2D* Texture = static_cast<ID3D11Texture2D*>(CurrentDepthBuffer->GetNativeResource());
+			ensure(HMD->CommitDepthBuffer(Texture));
+		}
+		else
+#endif
+		{
+			// We keep refs to the depth texture in the hmd, so this function is non-const.	
+			// RenderTexture_RenderThread should perhaps be refactored or made non-const.  
+			// But to get this into a hotfix of 4.23 we shall simply cast.
+			FWindowsMixedRealityHMD* nonconstthis = const_cast<FWindowsMixedRealityHMD*>(this);
+			nonconstthis->CreateHMDDepthTexture(RHICmdList);
+		}
 	}
 
 	// Create a BGRA backbuffer for rendering.
@@ -1315,22 +1364,47 @@ namespace WindowsMixedReality
 			return false;
 		}
 
-		FRHIResourceCreateInfo CreateInfo;
-
 		// Since our textures must be BGRA, this plugin did require a change to WindowsD3D11Device.cpp
 		// to add the D3D11_CREATE_DEVICE_BGRA_SUPPORT flag to the graphics device.
-		RHICreateTargetableShaderResource2D(
-			sizeX,
-			sizeY,
-			PF_B8G8R8A8, // must be BGRA
-			numMips,
-			flags,
-			targetableTextureFlags,
-			false,
-			CreateInfo,
-			outTargetableTexture,
-			outShaderResourceTexture);
+		FRHIResourceCreateInfo CreateInfo;
 
+#if PLATFORM_HOLOLENS
+		static const auto CVarMobileMultiViewDirect = IConsoleManager::Get().FindTConsoleVariableDataInt(TEXT("vr.MobileMultiView.Direct"));
+		const bool bIsMobileMultiViewDirectEnabled = (CVarMobileMultiViewDirect && CVarMobileMultiViewDirect->GetValueOnAnyThread() != 0);
+		if (bIsMobileMultiViewDirectEnabled)
+		{
+			FTexture2DArrayRHIRef texture, resource;
+			RHICreateTargetableShaderResource2DArray(
+				sizeX,
+				sizeY,
+				2,
+				PF_B8G8R8A8, // must be BGRA
+				numMips,
+				flags,
+				targetableTextureFlags,
+				CreateInfo,
+				texture,
+				resource);
+			outTargetableTexture = texture;
+			outShaderResourceTexture = resource;
+		}
+		else
+#endif
+		{
+			RHICreateTargetableShaderResource2D(
+				sizeX,
+				sizeY,
+				PF_B8G8R8A8, // must be BGRA
+				numMips,
+				flags,
+				targetableTextureFlags,
+				false,
+				CreateInfo,
+				outTargetableTexture,
+				outShaderResourceTexture);
+		}
+
+		CurrentBackBuffer = outTargetableTexture;
 		bNeedReallocateDepthTexture = true;
 
 		return true;
@@ -1353,20 +1427,49 @@ namespace WindowsMixedReality
 		// Current shader assumes far depth since scene depth uses far depth.
 		CreateInfo.ClearValueBinding = FClearValueBinding::DepthFar;
 
-		RHICreateTargetableShaderResource2D(
-			SizeX,
-			SizeY,
-			// Do not use input format - this will resolve to X32_TYPELESS_G8X24_UINT which cannot be used for a depthstencil buffer.
-			// DepthStencil will resolve to R32G8X24_TYPELESS which is usable for a depthstencil buffer!
-			PF_DepthStencil,
-			// Do not use input mips, this will resolve to 0 which will throw creating the texture.
-			1,
-			InTexFlags,
-			TargetableTextureFlags,
-			false,
-			CreateInfo,
-			OutTargetableTexture,
-			OutShaderResourceTexture);
+#if PLATFORM_HOLOLENS
+		static const auto CVarMobileMultiView = IConsoleManager::Get().FindTConsoleVariableDataInt(TEXT("vr.MobileMultiView"));
+		const bool bIsMobileMultiViewEnabled = (CVarMobileMultiView && CVarMobileMultiView->GetValueOnAnyThread() != 0);
+		if (bIsMobileMultiViewEnabled)
+		{
+			FIntPoint size = GetIdealRenderTargetSize();
+
+			FTexture2DArrayRHIRef texture, resource;
+			RHICreateTargetableShaderResource2DArray(
+				size.X,
+				size.Y,
+				2,
+				// Do not use input format - this will resolve to X32_TYPELESS_G8X24_UINT which cannot be used for a depthstencil buffer.
+				// DepthStencil will resolve to R32G8X24_TYPELESS which is usable for a depthstencil buffer!
+				PF_DepthStencil,
+				// Do not use input mips, this will resolve to 0 which will throw creating the texture.
+				1,
+				InTexFlags,
+				TargetableTextureFlags,
+				CreateInfo,
+				texture,
+				resource);
+			OutTargetableTexture = texture;
+			OutShaderResourceTexture = resource;
+		}
+		else
+#endif
+		{
+			RHICreateTargetableShaderResource2D(
+				SizeX,
+				SizeY,
+				// Do not use input format - this will resolve to X32_TYPELESS_G8X24_UINT which cannot be used for a depthstencil buffer.
+				// DepthStencil will resolve to R32G8X24_TYPELESS which is usable for a depthstencil buffer!
+				PF_DepthStencil,
+				// Do not use input mips, this will resolve to 0 which will throw creating the texture.
+				1,
+				InTexFlags,
+				TargetableTextureFlags,
+				false,
+				CreateInfo,
+				OutTargetableTexture,
+				OutShaderResourceTexture);
+		}
 
 		CurrentDepthBuffer = OutTargetableTexture;
 		bNeedReallocateDepthTexture = false;
diff --git a/Engine/Plugins/Runtime/WindowsMixedReality/Source/WindowsMixedRealityHMD/Private/WindowsMixedRealityHMD.h b/Engine/Plugins/Runtime/WindowsMixedReality/Source/WindowsMixedRealityHMD/Private/WindowsMixedRealityHMD.h
--- a/Engine/Plugins/Runtime/WindowsMixedReality/Source/WindowsMixedRealityHMD/Private/WindowsMixedRealityHMD.h
+++ b/Engine/Plugins/Runtime/WindowsMixedReality/Source/WindowsMixedRealityHMD/Private/WindowsMixedRealityHMD.h
@@ -219,6 +219,7 @@ namespace WindowsMixedReality
 		const float farPlaneDistance = 650.0f;
 
 		bool bNeedReallocateDepthTexture = false;
+		FTexture2DRHIRef CurrentBackBuffer;
 		FTexture2DRHIRef CurrentDepthBuffer;
 		void InitTrackingFrame();
 		TRefCountPtr<FWindowsMixedRealityCustomPresent> mCustomPresent = nullptr;
