Move HoloLens input sim injection for controller tracking status to the HMD class

The SpatialInput class only implements the IMotionController interface, but does not provide input sim data through the WMR function library. The function lib directly calls the HMD function, so implementing it on that level covers both code paths.

(Provided by Microsoft)

diff --git a/Engine/Plugins/Runtime/WindowsMixedReality/Source/WindowsMixedRealityHMD/Private/WindowsMixedRealityHMD.cpp b/Engine/Plugins/Runtime/WindowsMixedReality/Source/WindowsMixedRealityHMD/Private/WindowsMixedRealityHMD.cpp
--- a/Engine/Plugins/Runtime/WindowsMixedReality/Source/WindowsMixedRealityHMD/Private/WindowsMixedRealityHMD.cpp
+++ b/Engine/Plugins/Runtime/WindowsMixedReality/Source/WindowsMixedRealityHMD/Private/WindowsMixedRealityHMD.cpp
@@ -2043,7 +2043,16 @@ namespace WindowsMixedReality
 
 	HMDTrackingStatus FWindowsMixedRealityHMD::GetControllerTrackingStatus(HMDHand hand)
 	{
-		return HMD->GetControllerTrackingStatus(hand);
+#if WITH_INPUT_SIMULATION
+		if (auto* InputSim = UWindowsMixedRealityInputSimulationEngineSubsystem::GetInputSimulationIfEnabled())
+		{
+			return (HMDTrackingStatus)InputSim->GetControllerTrackingStatus((EControllerHand)hand);
+		}
+		else
+#endif
+		{
+			return HMD->GetControllerTrackingStatus(hand);
+		}
 	}
 
 	// GetControllerOrientationAndPosition() is pre-scaled to UE4 world scale
diff --git a/Engine/Plugins/Runtime/WindowsMixedReality/Source/WindowsMixedRealitySpatialInput/Private/WindowsMixedRealitySpatialInput.cpp b/Engine/Plugins/Runtime/WindowsMixedReality/Source/WindowsMixedRealitySpatialInput/Private/WindowsMixedRealitySpatialInput.cpp
--- a/Engine/Plugins/Runtime/WindowsMixedReality/Source/WindowsMixedRealitySpatialInput/Private/WindowsMixedRealitySpatialInput.cpp
+++ b/Engine/Plugins/Runtime/WindowsMixedReality/Source/WindowsMixedRealitySpatialInput/Private/WindowsMixedRealitySpatialInput.cpp
@@ -9,11 +9,6 @@
 #include "WindowsMixedRealitySpatialInputTypes.h"
 #include "WindowsMixedRealityAvailability.h"
 
-#if WITH_INPUT_SIMULATION
-#include "Engine/Engine.h"
-#include "WindowsMixedRealityInputSimulationEngineSubsystem.h"
-#endif
-
 #define LOCTEXT_NAMESPACE "WindowsMixedRealitySpatialInput"
 #define MotionControllerDeviceTypeName "WindowsMixedRealitySpatialInput"
 
@@ -805,12 +800,6 @@ namespace WindowsMixedReality
 	ETrackingStatus FWindowsMixedRealitySpatialInput::GetControllerTrackingStatus(const int32 ControllerIndex, const EControllerHand DeviceHand) const
 	{
 #if WITH_WINDOWS_MIXED_REALITY
-#if WITH_INPUT_SIMULATION
-		if (auto* InputSim = UWindowsMixedRealityInputSimulationEngineSubsystem::GetInputSimulationIfEnabled())
-		{
-			return InputSim->GetControllerTrackingStatus(DeviceHand);
-		}
-#endif
 		HMDHand hand = (HMDHand)((int)DeviceHand);
 		HMDTrackingStatus trackingStatus = FWindowsMixedRealityStatics::GetControllerTrackingStatus(hand);
 
