Better activation criteria for input simulation
Use the `bUseVRPreviewForPlayWorld` flag to determine if input sim should be enabled. This will use input sim in regular editor viewports, but also allow holographic remoting without having to disconnect the device.
Avoid restart requirement for toggling input simulation on and off. This requires that the engine subsystem is always created (only in editor). The IsInputSimulationEnabled or GetInputSimulationIfEnabled helper functions should be used to ensure the subsystem is only used when enabled.
(Submitted by Microsoft)

diff --git a/Engine/Plugins/Runtime/WindowsMixedReality/Source/WindowsMixedRealityInputSimulation/Private/WindowsMixedRealityInputSimulationEngineSubsystem.cpp b/Engine/Plugins/Runtime/WindowsMixedReality/Source/WindowsMixedRealityInputSimulation/Private/WindowsMixedRealityInputSimulationEngineSubsystem.cpp
--- a/Engine/Plugins/Runtime/WindowsMixedReality/Source/WindowsMixedRealityInputSimulation/Private/WindowsMixedRealityInputSimulationEngineSubsystem.cpp
+++ b/Engine/Plugins/Runtime/WindowsMixedReality/Source/WindowsMixedRealityInputSimulation/Private/WindowsMixedRealityInputSimulationEngineSubsystem.cpp
@@ -10,6 +10,10 @@
 #include "GameFramework/GameModeBase.h"
 #include "Misc/ScopeLock.h"
 
+#if WITH_EDITOR
+#include "Editor/EditorEngine.h"
+#endif
+
 #define LOCTEXT_NAMESPACE "WindowsMixedRealityInputSimulation"
 
 UWindowsMixedRealityInputSimulationEngineSubsystem* UWindowsMixedRealityInputSimulationEngineSubsystem::GetInputSimulationIfEnabled()
@@ -38,16 +42,17 @@ bool UWindowsMixedRealityInputSimulationEngineSubsystem::IsInputSimulationEnable
 		return false;
 	}
 
-	// Only if no HMD is connected
-	if (GEngine->XRSystem.Get()->GetHMDDevice()->IsHMDConnected())
+	// Only if playing inside the editor viewport
+	if (GIsEditor)
 	{
-		return false;
+		UEditorEngine* EdEngine = Cast<UEditorEngine>(GEngine);
+
+		// We take bUseVRPreviewForPlayWorld as an indication that we're running using Holographic Remoting
+		return !EdEngine->bUseVRPreviewForPlayWorld;
 	}
+#endif
 
-	return true;
-#else
 	return false;
-#endif
 }
 
 bool UWindowsMixedRealityInputSimulationEngineSubsystem::HasPositionalTracking() const
@@ -163,14 +168,17 @@ void UWindowsMixedRealityInputSimulationEngineSubsystem::UpdateSimulatedData(
 	const FWindowsMixedRealityInputSimulationHandState& LeftHandState,
 	const FWindowsMixedRealityInputSimulationHandState& RightHandState)
 {
-	FRWScopeLock WriteLock(DataMutex, SLT_Write);
+	if (IsInputSimulationEnabled())
+	{
+		FRWScopeLock WriteLock(DataMutex, SLT_Write);
 
-	bHasPositionalTracking = HasTracking;
-	HeadOrientation = NewHeadOrientation;
-	HeadPosition = NewHeadPosition;
+		bHasPositionalTracking = HasTracking;
+		HeadOrientation = NewHeadOrientation;
+		HeadPosition = NewHeadPosition;
 
-	UpdateSimulatedHandState(EControllerHand::Left, LeftHandState);
-	UpdateSimulatedHandState(EControllerHand::Right, RightHandState);
+		UpdateSimulatedHandState(EControllerHand::Left, LeftHandState);
+		UpdateSimulatedHandState(EControllerHand::Right, RightHandState);
+	}
 }
 
 void UWindowsMixedRealityInputSimulationEngineSubsystem::UpdateSimulatedHandState(
@@ -187,18 +195,8 @@ void UWindowsMixedRealityInputSimulationEngineSubsystem::UpdateSimulatedHandStat
 
 bool UWindowsMixedRealityInputSimulationEngineSubsystem::ShouldCreateSubsystem(UObject* Outer) const
 {
+	// Enable input simulation only in the editor
 #if WITH_EDITOR
-	// Only use input simulation when enabled
-	const UWindowsMixedRealityRuntimeSettings* WMRSettings = UWindowsMixedRealityRuntimeSettings::Get();
-	if (!WMRSettings->bEnableInputSimulation)
-	{
-		return false;
-	}
-
-	// Note: XRSystem is created later than engine subsystem,
-	// so cannot check if XRSystem exists or HMD is connected at this point.
-	// IsInputSimulationEnabled can be used at runtime to check these additional conditions.
-
 	return true;
 #else
 	return false;
diff --git a/Engine/Plugins/Runtime/WindowsMixedReality/Source/WindowsMixedRealityInputSimulation/WindowsMixedRealityInputSimulation.Build.cs b/Engine/Plugins/Runtime/WindowsMixedReality/Source/WindowsMixedRealityInputSimulation/WindowsMixedRealityInputSimulation.Build.cs
--- a/Engine/Plugins/Runtime/WindowsMixedReality/Source/WindowsMixedRealityInputSimulation/WindowsMixedRealityInputSimulation.Build.cs
+++ b/Engine/Plugins/Runtime/WindowsMixedReality/Source/WindowsMixedRealityInputSimulation/WindowsMixedRealityInputSimulation.Build.cs
@@ -33,6 +33,11 @@ namespace UnrealBuildTool.Rules
 					"WindowsMixedRealityRuntimeSettings",
 				}
 			);
-		}
+
+            if (Target.bBuildEditor)
+            {
+                PrivateDependencyModuleNames.Add("UnrealEd");
+            }
+        }
 	}
 }
diff --git a/Engine/Plugins/Runtime/WindowsMixedReality/Source/WindowsMixedRealityRuntimeSettings/Classes/WindowsMixedRealityRuntimeSettings.h b/Engine/Plugins/Runtime/WindowsMixedReality/Source/WindowsMixedRealityRuntimeSettings/Classes/WindowsMixedRealityRuntimeSettings.h
--- a/Engine/Plugins/Runtime/WindowsMixedReality/Source/WindowsMixedRealityRuntimeSettings/Classes/WindowsMixedRealityRuntimeSettings.h
+++ b/Engine/Plugins/Runtime/WindowsMixedReality/Source/WindowsMixedRealityRuntimeSettings/Classes/WindowsMixedRealityRuntimeSettings.h
@@ -48,6 +48,6 @@ public:
 	UPROPERTY(GlobalConfig, EditAnywhere, Category = "Holographic Remoting", Meta = (DisplayName = "HoloLens 1 Remoting", Tooltip = "If True remoting connect will assume the device being connected is a HL1, if False HL2 is assumed.  If you chose wrong remoting will fail to connect."))
 	bool IsHoloLens1Remoting = false;
 
-	UPROPERTY(GlobalConfig, EditAnywhere, Category = "Input Simulation", Meta = (ConfigRestartRequired = true, DisplayName = "Enable Input Simulation", Tooltip = "Enable simulation of AR input in the editor when no HMD is connected."))
+	UPROPERTY(GlobalConfig, EditAnywhere, Category = "Input Simulation", Meta = (DisplayName = "Enable Input Simulation", Tooltip = "Enable simulation of AR input in the editor when no HMD is connected."))
 	bool bEnableInputSimulation = true;
 };
