Add support for the hand pointer pose in input simulation
(Submitted by Microsoft)

diff --git a/Engine/Plugins/Runtime/WindowsMixedReality/Source/WindowsMixedRealityHMD/Private/WindowsMixedRealityHMD.cpp b/Engine/Plugins/Runtime/WindowsMixedReality/Source/WindowsMixedRealityHMD/Private/WindowsMixedRealityHMD.cpp
--- a/Engine/Plugins/Runtime/WindowsMixedReality/Source/WindowsMixedRealityHMD/Private/WindowsMixedRealityHMD.cpp
+++ b/Engine/Plugins/Runtime/WindowsMixedReality/Source/WindowsMixedRealityHMD/Private/WindowsMixedRealityHMD.cpp
@@ -2174,21 +2174,41 @@ namespace WindowsMixedReality
 
 	void FWindowsMixedRealityHMD::GetPointerPose(EControllerHand hand, PointerPoseInfo& pi)
 	{
-		HMD->GetPointerPose((HMDHand)hand, pi);
+#if WITH_INPUT_SIMULATION
+		if (auto* InputSim = UWindowsMixedRealityInputSimulationEngineSubsystem::GetInputSimulationIfEnabled())
+		{
+			FWindowsMixedRealityInputSimulationPointerPose InputSimPointerPose;
+			if (InputSim->GetHandPointerPose(hand, InputSimPointerPose))
+			{
+				FVector pos = InputSimPointerPose.Origin;
+				FVector dir = InputSimPointerPose.Direction;
+				FVector up = InputSimPointerPose.Up;
+				FQuat rot = InputSimPointerPose.Orientation;
+				pi.origin = DirectX::XMFLOAT3(pos.X, pos.Y, pos.Z);
+				pi.orientation = DirectX::XMFLOAT4(rot.X, rot.Y, rot.Z, rot.W);
+				pi.up = DirectX::XMFLOAT3(up.X, up.Y, up.Z);
+				pi.direction = DirectX::XMFLOAT3(dir.X, dir.Y, dir.Z);
+			}
+		}
+		else
+#endif
+		{
+			HMD->GetPointerPose((HMDHand)hand, pi);
 
-		FTransform TrackingSpaceTransformOrigin(WMRUtility::FromMixedRealityQuaternion(pi.orientation), WMRUtility::FromMixedRealityVector(pi.origin) * GetWorldToMetersScale());
+			FTransform TrackingSpaceTransformOrigin(WMRUtility::FromMixedRealityQuaternion(pi.orientation), WMRUtility::FromMixedRealityVector(pi.origin) * GetWorldToMetersScale());
 
-		FTransform tOrigin = (TrackingSpaceTransformOrigin * CachedTrackingToWorld);
-		FVector pos = tOrigin.GetLocation();
-		FQuat rot = tOrigin.GetRotation();
+			FTransform tOrigin = (TrackingSpaceTransformOrigin * CachedTrackingToWorld);
+			FVector pos = tOrigin.GetLocation();
+			FQuat rot = tOrigin.GetRotation();
 
-		FVector up = CachedTrackingToWorld.GetRotation() * WMRUtility::FromMixedRealityVector(pi.up);
-		FVector dir = CachedTrackingToWorld.GetRotation() * WMRUtility::FromMixedRealityVector(pi.direction);
+			FVector up = CachedTrackingToWorld.GetRotation() * WMRUtility::FromMixedRealityVector(pi.up);
+			FVector dir = CachedTrackingToWorld.GetRotation() * WMRUtility::FromMixedRealityVector(pi.direction);
 
-		pi.origin = DirectX::XMFLOAT3(pos.X, pos.Y, pos.Z);
-		pi.orientation = DirectX::XMFLOAT4(rot.X, rot.Y, rot.Z, rot.W);
-		pi.up = DirectX::XMFLOAT3(up.X, up.Y, up.Z);
-		pi.direction = DirectX::XMFLOAT3(dir.X, dir.Y, dir.Z);
+			pi.origin = DirectX::XMFLOAT3(pos.X, pos.Y, pos.Z);
+			pi.orientation = DirectX::XMFLOAT4(rot.X, rot.Y, rot.Z, rot.W);
+			pi.up = DirectX::XMFLOAT3(up.X, up.Y, up.Z);
+			pi.direction = DirectX::XMFLOAT3(dir.X, dir.Y, dir.Z);
+		}
 	}
 
 #endif
diff --git a/Engine/Plugins/Runtime/WindowsMixedReality/Source/WindowsMixedRealityInputSimulation/Private/WindowsMixedRealityInputSimulationEngineSubsystem.cpp b/Engine/Plugins/Runtime/WindowsMixedReality/Source/WindowsMixedRealityInputSimulation/Private/WindowsMixedRealityInputSimulationEngineSubsystem.cpp
--- a/Engine/Plugins/Runtime/WindowsMixedReality/Source/WindowsMixedRealityInputSimulation/Private/WindowsMixedRealityInputSimulationEngineSubsystem.cpp
+++ b/Engine/Plugins/Runtime/WindowsMixedReality/Source/WindowsMixedRealityInputSimulation/Private/WindowsMixedRealityInputSimulationEngineSubsystem.cpp
@@ -161,6 +161,20 @@ bool UWindowsMixedRealityInputSimulationEngineSubsystem::GetPressState(EControll
 	return false;
 }
 
+bool UWindowsMixedRealityInputSimulationEngineSubsystem::GetHandPointerPose(EControllerHand Hand, FWindowsMixedRealityInputSimulationPointerPose& OutPointerPose) const
+{
+	FRWScopeLock ReadLock(DataMutex, SLT_ReadOnly);
+	if (const FWindowsMixedRealityInputSimulationHandState* handState = GetHandState(Hand))
+	{
+		if (handState->bHasPointerPose)
+		{
+			OutPointerPose = handState->PointerPose;
+			return true;
+		}
+	}
+	return false;
+}
+
 void UWindowsMixedRealityInputSimulationEngineSubsystem::UpdateSimulatedData(
 	bool HasTracking,
 	const FQuat& NewHeadOrientation,
diff --git a/Engine/Plugins/Runtime/WindowsMixedReality/Source/WindowsMixedRealityInputSimulation/Public/WindowsMixedRealityInputSimulationEngineSubsystem.h b/Engine/Plugins/Runtime/WindowsMixedReality/Source/WindowsMixedRealityInputSimulation/Public/WindowsMixedRealityInputSimulationEngineSubsystem.h
--- a/Engine/Plugins/Runtime/WindowsMixedReality/Source/WindowsMixedRealityInputSimulation/Public/WindowsMixedRealityInputSimulationEngineSubsystem.h
+++ b/Engine/Plugins/Runtime/WindowsMixedReality/Source/WindowsMixedRealityInputSimulation/Public/WindowsMixedRealityInputSimulationEngineSubsystem.h
@@ -13,8 +13,16 @@
 
 #include "WindowsMixedRealityInputSimulationEngineSubsystem.generated.h"
 
+struct WINDOWSMIXEDREALITYINPUTSIMULATION_API FWindowsMixedRealityInputSimulationPointerPose
+{
+	FVector Origin = FVector::ZeroVector;
+	FVector Direction = FVector::ForwardVector;
+	FVector Up = FVector::UpVector;
+	FQuat Orientation = FQuat::Identity;
+};
+
 /** Simulated data for one hand. */
-struct FWindowsMixedRealityInputSimulationHandState
+struct WINDOWSMIXEDREALITYINPUTSIMULATION_API FWindowsMixedRealityInputSimulationHandState
 {
 	typedef TStaticArray<FTransform, (uint32)EWMRHandKeypointCount> KeypointTransformArray;
 	typedef TStaticArray<float, (uint32)EWMRHandKeypointCount> KeypointRadiusArray;
@@ -28,6 +36,9 @@ struct FWindowsMixedRealityInputSimulationHandState
 
 	ButtonStateArray IsButtonPressed = ButtonStateArray();
 	ButtonStateArray PrevButtonPressed = ButtonStateArray();
+
+	bool bHasPointerPose = false;
+	FWindowsMixedRealityInputSimulationPointerPose PointerPose;
 };
 
 /** Engine subsystem that stores input simulation data for access by the XR device. */
@@ -78,6 +89,10 @@ public:
 	 */
 	bool GetPressState(EControllerHand Hand, EHMDInputControllerButtons Button, bool OnlyRegisterClicks, bool& OutPressState) const;
 
+	/** Get the pointer pose for a hand.
+	 */
+	bool GetHandPointerPose(EControllerHand Hand, FWindowsMixedRealityInputSimulationPointerPose& OutPointerPose) const;
+
 	/** Update the simulated data. */
 	void UpdateSimulatedData(
 		bool HasTracking,
