Fix deprecation warning in new Input Simulation code from MS.

diff --git a/Engine/Plugins/Runtime/WindowsMixedReality/Source/WindowsMixedRealityInputSimulation/Private/WindowsMixedRealityInputSimulationEngineSubsystem.cpp b/Engine/Plugins/Runtime/WindowsMixedReality/Source/WindowsMixedRealityInputSimulation/Private/WindowsMixedRealityInputSimulationEngineSubsystem.cpp
--- a/Engine/Plugins/Runtime/WindowsMixedReality/Source/WindowsMixedRealityInputSimulation/Private/WindowsMixedRealityInputSimulationEngineSubsystem.cpp
+++ b/Engine/Plugins/Runtime/WindowsMixedReality/Source/WindowsMixedRealityInputSimulation/Private/WindowsMixedRealityInputSimulationEngineSubsystem.cpp
@@ -48,7 +48,8 @@ bool UWindowsMixedRealityInputSimulationEngineSubsystem::IsInputSimulationEnable
 		UEditorEngine* EdEngine = Cast<UEditorEngine>(GEngine);
 
 		// We take bUseVRPreviewForPlayWorld as an indication that we're running using Holographic Remoting
-		return !EdEngine->bUseVRPreviewForPlayWorld;
+		//return FRequestPlaySessionParams::SessionPreviewTypeOverride 
+		return !(EdEngine->GetPlaySessionRequest()->SessionPreviewTypeOverride == EPlaySessionPreviewType::VRPreview);
 	}
 #endif
 
