UE-91195 Hololens - input simulator active and crashes in vr preview
-Fixed logic for detecting VRPreview

diff --git a/Engine/Plugins/Runtime/WindowsMixedReality/Source/WindowsMixedRealityInputSimulation/Private/WindowsMixedRealityInputSimulationEngineSubsystem.cpp b/Engine/Plugins/Runtime/WindowsMixedReality/Source/WindowsMixedRealityInputSimulation/Private/WindowsMixedRealityInputSimulationEngineSubsystem.cpp
--- a/Engine/Plugins/Runtime/WindowsMixedReality/Source/WindowsMixedRealityInputSimulation/Private/WindowsMixedRealityInputSimulationEngineSubsystem.cpp
+++ b/Engine/Plugins/Runtime/WindowsMixedReality/Source/WindowsMixedRealityInputSimulation/Private/WindowsMixedRealityInputSimulationEngineSubsystem.cpp
@@ -47,9 +47,16 @@ bool UWindowsMixedRealityInputSimulationEngineSubsystem::IsInputSimulationEnable
 	{
 		UEditorEngine* EdEngine = Cast<UEditorEngine>(GEngine);
 
-		// We take bUseVRPreviewForPlayWorld as an indication that we're running using Holographic Remoting
-		//return FRequestPlaySessionParams::SessionPreviewTypeOverride 
-		return !(EdEngine->GetPlaySessionRequest()->SessionPreviewTypeOverride == EPlaySessionPreviewType::VRPreview);
+		if (EdEngine->GetPlayInEditorSessionInfo().IsSet())
+		{
+			// If we are in a VRPreview we are remoting or connected to wmr, and don't want simulated input.			
+			return EdEngine->GetPlayInEditorSessionInfo()->OriginalRequestParams.SessionPreviewTypeOverride != EPlaySessionPreviewType::VRPreview;
+		}
+		else
+		{
+			return false;
+		}
+
 	}
 #endif
 
