UE-91210
Hololens - restore wmr/hololensar plugin dependency direction to previous

diff --git a/Engine/Plugins/Runtime/AR/Microsoft/HoloLensAR/HoloLensAR.uplugin b/Engine/Plugins/Runtime/AR/Microsoft/HoloLensAR/HoloLensAR.uplugin
--- a/Engine/Plugins/Runtime/AR/Microsoft/HoloLensAR/HoloLensAR.uplugin
+++ b/Engine/Plugins/Runtime/AR/Microsoft/HoloLensAR/HoloLensAR.uplugin
@@ -44,15 +44,5 @@
 				"Win64"
 			]
 		}
-	],
-	"Plugins": [
-		{
-			"Name": "WindowsMixedReality",
-			"Enabled": true,
-			"WhitelistPlatforms": [
-				"Win64",
-				"HoloLens"
-			]
-		}
 	]
 }
\ No newline at end of file
diff --git a/Engine/Plugins/Runtime/AR/Microsoft/HoloLensAR/Source/HoloLensAR/Public/HoloLensModule.h b/Engine/Plugins/Runtime/AR/Microsoft/HoloLensAR/Source/HoloLensAR/Public/HoloLensModule.h
--- a/Engine/Plugins/Runtime/AR/Microsoft/HoloLensAR/Source/HoloLensAR/Public/HoloLensModule.h
+++ b/Engine/Plugins/Runtime/AR/Microsoft/HoloLensAR/Source/HoloLensAR/Public/HoloLensModule.h
@@ -5,7 +5,6 @@
 #include "Modules/ModuleManager.h"
 #include "Features/IModularFeature.h"
 #include "HoloLensARSystem.h"
-#include "WindowsMixedRealityHMD/Public/IWindowsMixedRealityHMDPlugin.h"
 
 DECLARE_LOG_CATEGORY_EXTERN(LogHoloLensAR, Log, All);
 
@@ -16,16 +15,16 @@ namespace WindowsMixedReality
 }
 
 class HOLOLENSAR_API FHoloLensModuleAR :
-	public IHoloLensModuleARInterface
+	public IModuleInterface
 {
 public:
-	static TSharedPtr<FHoloLensARSystem, ESPMode::ThreadSafe> GetHoloLensARSystem();
-
 	/** Used to init our AR system */
-	virtual class IARSystemSupport* CreateARSystem() override;
+	static IARSystemSupport* CreateARSystem();
 
-	virtual void SetTrackingSystem(TSharedPtr<class FXRTrackingSystemBase, ESPMode::ThreadSafe> InTrackingSystem) override;
-	virtual void SetInterop(WindowsMixedReality::MixedRealityInterop* InWMRInterop) override;
+	static TSharedPtr<FHoloLensARSystem, ESPMode::ThreadSafe> GetHoloLensARSystem();
+	
+	static void SetTrackingSystem(TSharedPtr<FXRTrackingSystemBase, ESPMode::ThreadSafe> InTrackingSystem);
+	static void SetInterop(WindowsMixedReality::MixedRealityInterop* InWMRInterop);
 
 	virtual void StartupModule() override;
 
diff --git a/Engine/Plugins/Runtime/WindowsMixedReality/Source/WindowsMixedRealityHMD/Private/WindowsMixedRealityHMD.cpp b/Engine/Plugins/Runtime/WindowsMixedReality/Source/WindowsMixedRealityHMD/Private/WindowsMixedRealityHMD.cpp
--- a/Engine/Plugins/Runtime/WindowsMixedReality/Source/WindowsMixedRealityHMD/Private/WindowsMixedRealityHMD.cpp
+++ b/Engine/Plugins/Runtime/WindowsMixedReality/Source/WindowsMixedRealityHMD/Private/WindowsMixedRealityHMD.cpp
@@ -152,11 +152,7 @@ namespace WindowsMixedReality
 			if (HMD)
 			{
 #if SUPPORTS_WINDOWS_MIXED_REALITY_AR
-				HololensARModule = (IHoloLensModuleARInterface*)FModuleManager::Get().LoadModule("HoloLensAR");
-				if (HololensARModule)
-				{
-					HololensARModule->SetInterop(nullptr);
-				}
+				FHoloLensModuleAR::SetInterop(nullptr);
 #endif
 				HMD->Dispose(true);
 				delete HMD;
@@ -178,7 +174,6 @@ namespace WindowsMixedReality
 
 #if WITH_WINDOWS_MIXED_REALITY
 		MixedRealityInterop* HMD = nullptr;
-		IHoloLensModuleARInterface * HololensARModule = nullptr;
 #endif
 	};
 
@@ -189,22 +184,14 @@ namespace WindowsMixedReality
 		{
 			IARSystemSupport* ARSystem = nullptr;
 #if SUPPORTS_WINDOWS_MIXED_REALITY_AR
-			HololensARModule = (IHoloLensModuleARInterface *)FModuleManager::Get().LoadModule("HoloLensAR");
-			if (HololensARModule)
-			{
-				ARSystem = HololensARModule->CreateARSystem();
-			}
+			ARSystem = FHoloLensModuleAR::CreateARSystem();
 #endif
 			auto WindowsMRHMD = FSceneViewExtensions::NewExtension<WindowsMixedReality::FWindowsMixedRealityHMD>(ARSystem, HMD);
 			if (WindowsMRHMD->IsInitialized())
 			{
 #if SUPPORTS_WINDOWS_MIXED_REALITY_AR
-				HololensARModule = (IHoloLensModuleARInterface*)FModuleManager::Get().LoadModule("HoloLensAR");
-				if (HololensARModule)
-				{
-					HololensARModule->SetTrackingSystem(WindowsMRHMD);
-					HololensARModule->SetInterop(HMD);
-				}
+				FHoloLensModuleAR::SetTrackingSystem(WindowsMRHMD);
+				FHoloLensModuleAR::SetInterop(HMD);
 				// Register the AR modular features
 				WindowsMRHMD->GetARCompositionComponent()->InitializeARSystem();
 #endif
diff --git a/Engine/Plugins/Runtime/WindowsMixedReality/Source/WindowsMixedRealityHMD/Public/IWindowsMixedRealityHMDPlugin.h b/Engine/Plugins/Runtime/WindowsMixedReality/Source/WindowsMixedRealityHMD/Public/IWindowsMixedRealityHMDPlugin.h
--- a/Engine/Plugins/Runtime/WindowsMixedReality/Source/WindowsMixedRealityHMD/Public/IWindowsMixedRealityHMDPlugin.h
+++ b/Engine/Plugins/Runtime/WindowsMixedReality/Source/WindowsMixedRealityHMD/Public/IWindowsMixedRealityHMDPlugin.h
@@ -41,16 +41,3 @@ public:
 		return nullptr;
 	}
 };
-class IHoloLensModuleARInterface :
-	public IModuleInterface
-{
-public:
-	/** Used to init our AR system */
-	virtual class IARSystemSupport* CreateARSystem() = 0;
-
-
-	virtual void SetTrackingSystem(TSharedPtr<class FXRTrackingSystemBase, ESPMode::ThreadSafe> InTrackingSystem) = 0;
-	virtual void SetInterop(WindowsMixedReality::MixedRealityInterop* InWMRInterop) = 0;
-
-};
-
diff --git a/Engine/Plugins/Runtime/WindowsMixedReality/Source/WindowsMixedRealityHMD/WindowsMixedRealityHMD.Build.cs b/Engine/Plugins/Runtime/WindowsMixedReality/Source/WindowsMixedRealityHMD/WindowsMixedRealityHMD.Build.cs
--- a/Engine/Plugins/Runtime/WindowsMixedReality/Source/WindowsMixedRealityHMD/WindowsMixedRealityHMD.Build.cs
+++ b/Engine/Plugins/Runtime/WindowsMixedReality/Source/WindowsMixedRealityHMD/WindowsMixedRealityHMD.Build.cs
@@ -83,6 +83,11 @@ namespace UnrealBuildTool.Rules
 					}
 					);
 
+				if (Target.Platform == UnrealTargetPlatform.Win64 ||
+					Target.Platform == UnrealTargetPlatform.HoloLens)
+				{
+					PrivateDependencyModuleNames.Add("HoloLensAR");
+				}
 
 				if (Target.bBuildEditor == true)
 				{
diff --git a/Engine/Plugins/Runtime/WindowsMixedReality/WindowsMixedReality.uplugin b/Engine/Plugins/Runtime/WindowsMixedReality/WindowsMixedReality.uplugin
--- a/Engine/Plugins/Runtime/WindowsMixedReality/WindowsMixedReality.uplugin
+++ b/Engine/Plugins/Runtime/WindowsMixedReality/WindowsMixedReality.uplugin
@@ -98,6 +98,14 @@
 		{
 			"Name": "LiveLink",
 			"Enabled": true
+		},
+		{
+			"Name": "HoloLensAR",
+			"Enabled": true,
+			"WhitelistPlatforms": [
+				"Win64",
+				"HoloLens"
+			]
 		}
 	]
 }
