PR #6815: Remove duplicate TrackingToWorldTransform in HoloLens Hand Joints (Contributed by [REDACTED_PERSON])

diff --git a/Engine/Plugins/Runtime/WindowsMixedReality/Source/WindowsMixedRealityHandTracking/Private/WindowsMixedRealityHandTrackingFunctionLibrary.cpp b/Engine/Plugins/Runtime/WindowsMixedReality/Source/WindowsMixedRealityHandTracking/Private/WindowsMixedRealityHandTrackingFunctionLibrary.cpp
--- a/Engine/Plugins/Runtime/WindowsMixedReality/Source/WindowsMixedRealityHandTracking/Private/WindowsMixedRealityHandTrackingFunctionLibrary.cpp
+++ b/Engine/Plugins/Runtime/WindowsMixedReality/Source/WindowsMixedRealityHandTracking/Private/WindowsMixedRealityHandTrackingFunctionLibrary.cpp
@@ -22,7 +22,7 @@ bool UWindowsMixedRealityHandTrackingFunctionLibrary::GetHandJointTransform(ECon
 
 		if (bSuccess)
 		{
-			OutTransform = KeyPointTransform * UHeadMountedDisplayFunctionLibrary::GetTrackingToWorldTransform(GWorld);
+			OutTransform = KeyPointTransform;
 			HandTracking->GetKeypointRadius(Hand, Keypoint, OutRadius);
 			return true;
 		}
