Fix ShaderCompiler incorrectly determining Instanced Multi-View support
Only use layered swapchains for the scene render targets in the WMR plugin
Only attempt to use instanced mobile multi-view for Hololens

diff --git a/Engine/Plugins/Runtime/OpenXR/Source/OpenXRHMD/Private/OpenXRHMD.cpp b/Engine/Plugins/Runtime/OpenXR/Source/OpenXRHMD/Private/OpenXRHMD.cpp
--- a/Engine/Plugins/Runtime/OpenXR/Source/OpenXRHMD/Private/OpenXRHMD.cpp
+++ b/Engine/Plugins/Runtime/OpenXR/Source/OpenXRHMD/Private/OpenXRHMD.cpp
@@ -893,15 +893,16 @@ FOpenXRHMD::FOpenXRHMD(const FAutoRegister& AutoRegister, XrInstance InInstance,
 
 #if PLATFORM_HOLOLENS
 	bDepthExtensionSupported = bHiddenAreaMaskSupported = false;
-#else
-	bDepthExtensionSupported = IsExtensionEnabled(XR_KHR_COMPOSITION_LAYER_DEPTH_EXTENSION_NAME);
-	bHiddenAreaMaskSupported = IsExtensionEnabled(XR_KHR_VISIBILITY_MASK_EXTENSION_NAME);
-#endif
 
 	static const auto CVarMobileMultiView = IConsoleManager::Get().FindTConsoleVariableDataInt(TEXT("vr.MobileMultiView"));
 	static const auto CVarMobileHDR = IConsoleManager::Get().FindTConsoleVariableDataInt(TEXT("r.MobileHDR"));
 	const bool bMobileHDR = (CVarMobileHDR && CVarMobileHDR->GetValueOnAnyThread() != 0);
-	bIsMobileMultiViewEnabled = (GSupportsMobileMultiView && !bMobileHDR) && (CVarMobileMultiView && CVarMobileMultiView->GetValueOnAnyThread() != 0);
+	bIsMobileMultiViewEnabled = (GRHISupportsArrayIndexFromAnyShader && !bMobileHDR) && (CVarMobileMultiView && CVarMobileMultiView->GetValueOnAnyThread() != 0);
+#else
+	bDepthExtensionSupported = IsExtensionEnabled(XR_KHR_COMPOSITION_LAYER_DEPTH_EXTENSION_NAME);
+	bHiddenAreaMaskSupported = IsExtensionEnabled(XR_KHR_VISIBILITY_MASK_EXTENSION_NAME);
+	bIsMobileMultiViewEnabled = false;
+#endif
 
 	{
 		// Enumerate the viewport configurations
diff --git a/Engine/Plugins/Runtime/WindowsMixedReality/Source/WindowsMixedRealityHMD/Private/WindowsMixedRealityHMD.cpp b/Engine/Plugins/Runtime/WindowsMixedReality/Source/WindowsMixedRealityHMD/Private/WindowsMixedRealityHMD.cpp
--- a/Engine/Plugins/Runtime/WindowsMixedReality/Source/WindowsMixedRealityHMD/Private/WindowsMixedRealityHMD.cpp
+++ b/Engine/Plugins/Runtime/WindowsMixedReality/Source/WindowsMixedRealityHMD/Private/WindowsMixedRealityHMD.cpp
@@ -1368,7 +1368,8 @@ namespace WindowsMixedReality
 		FRHIResourceCreateInfo CreateInfo;
 
 #if PLATFORM_HOLOLENS
-		if (bIsMobileMultiViewEnabled)
+        // TODO: Remove index < 16 condition once FWindowsMixedRealityViewport has been removed
+		if (index < 16 && bIsMobileMultiViewEnabled)
 		{
 			FTexture2DArrayRHIRef texture, resource;
 			RHICreateTargetableShaderResource2DArray(
@@ -1384,6 +1385,8 @@ namespace WindowsMixedReality
 				resource);
 			outTargetableTexture = texture;
 			outShaderResourceTexture = resource;
+
+			CurrentBackBuffer = outTargetableTexture;
 		}
 		else
 #endif
@@ -1401,7 +1404,6 @@ namespace WindowsMixedReality
 				outShaderResourceTexture);
 		}
 
-		CurrentBackBuffer = outTargetableTexture;
 		bNeedReallocateDepthTexture = true;
 
 		return true;
@@ -1774,10 +1776,12 @@ namespace WindowsMixedReality
 		static const FName RendererModuleName("Renderer");
 		RendererModule = FModuleManager::GetModulePtr<IRendererModule>(RendererModuleName);
 
+#if PLATFORM_HOLOLENS
 		static const auto CVarMobileMultiView = IConsoleManager::Get().FindTConsoleVariableDataInt(TEXT("vr.MobileMultiView"));
 		static const auto CVarMobileHDR = IConsoleManager::Get().FindTConsoleVariableDataInt(TEXT("r.MobileHDR"));
 		const bool bMobileHDR = (CVarMobileHDR && CVarMobileHDR->GetValueOnAnyThread() != 0);
-		bIsMobileMultiViewEnabled = (GSupportsMobileMultiView && !bMobileHDR) && (CVarMobileMultiView && CVarMobileMultiView->GetValueOnAnyThread() != 0);
+		bIsMobileMultiViewEnabled = (GRHISupportsArrayIndexFromAnyShader && !bMobileHDR) && (CVarMobileMultiView && CVarMobileMultiView->GetValueOnAnyThread() != 0);
+#endif
 
 		HiddenAreaMesh.SetNum(2);
 		VisibleAreaMesh.SetNum(2);
diff --git a/Engine/Plugins/Runtime/WindowsMixedReality/Source/WindowsMixedRealityRHI/Private/WindowsMixedRealityRHI.cpp b/Engine/Plugins/Runtime/WindowsMixedReality/Source/WindowsMixedRealityRHI/Private/WindowsMixedRealityRHI.cpp
--- a/Engine/Plugins/Runtime/WindowsMixedReality/Source/WindowsMixedRealityRHI/Private/WindowsMixedRealityRHI.cpp
+++ b/Engine/Plugins/Runtime/WindowsMixedReality/Source/WindowsMixedRealityRHI/Private/WindowsMixedRealityRHI.cpp
@@ -267,10 +267,11 @@ void FWindowsMixedRealityViewport::UpdateBackBuffer()
 	BackBuffer = nullptr;
 	if (GEngine && GEngine->StereoRenderingDevice && GEngine->StereoRenderingDevice->GetRenderTargetManager())
 	{
+        // HACK: Allocate with index 16 to avoid layered RTs from being allocated for viewports that don't need it
 		if (GEngine
 			->StereoRenderingDevice
 			->GetRenderTargetManager()
-			->AllocateRenderTargetTexture(0, SizeX, SizeY, PixelFormat, 1, TexCreate_None, TexCreate_RenderTargetable, OutTargetableTexture, OutShaderResourceTexture))
+			->AllocateRenderTargetTexture(16, SizeX, SizeY, PixelFormat, 1, TexCreate_None, TexCreate_RenderTargetable, OutTargetableTexture, OutShaderResourceTexture))
 		{
 			BackBuffer = (FD3D11Texture2D*)(OutTargetableTexture.GetReference());
 		}
diff --git a/Engine/Shaders/Public/ShaderVersion.ush b/Engine/Shaders/Public/ShaderVersion.ush
--- a/Engine/Shaders/Public/ShaderVersion.ush
+++ b/Engine/Shaders/Public/ShaderVersion.ush
@@ -3,4 +3,4 @@
 // in Platform.ush (which should be included in any shader) it allows to invalidate the shader DDC.
 // 
 // If you are merging streams and there is a conflict with this GUID you should make a new GUID rather than taking one or the other.
-// GUID = B83EB8AE-77C7-404F-BD81-9F1CCC330F59
+// GUID = E54ED269-25A1-4064-884C-D3017370E4FA
diff --git a/Engine/Source/Runtime/Engine/Private/SceneView.cpp b/Engine/Source/Runtime/Engine/Private/SceneView.cpp
--- a/Engine/Source/Runtime/Engine/Private/SceneView.cpp
+++ b/Engine/Source/Runtime/Engine/Private/SceneView.cpp
@@ -28,6 +28,7 @@
 #include "RenderUtils.h"
 
 DEFINE_LOG_CATEGORY(LogBufferVisualization);
+DEFINE_LOG_CATEGORY(LogMultiView);
 
 DECLARE_CYCLE_STAT(TEXT("StartFinalPostprocessSettings"), STAT_StartFinalPostprocessSettings, STATGROUP_Engine);
 DECLARE_CYCLE_STAT(TEXT("OverridePostProcessSettings"), STAT_OverridePostProcessSettings, STATGROUP_Engine);
@@ -764,15 +765,12 @@ FSceneView::FSceneView(const FSceneViewInitOptions& InitOptions)
 	bIsMobileMultiViewEnabled = bUsingMobileRenderer && bSkipPostprocessing && (MobileMultiViewCVar && MobileMultiViewCVar->GetValueOnAnyThread() != 0);
 	if (bIsMobileMultiViewEnabled && !RHISupportsMobileMultiView(ShaderPlatform))
 	{
-		// Native mobile multi-view is not supported, attempt to fall back to instancing
-		if (GRHISupportsArrayIndexFromAnyShader)
+		// Native mobile multi-view is not supported, attempt to fall back to instancing on compatible RHIs
+		if (RHISupportsInstancedStereo(ShaderPlatform) && !GRHISupportsArrayIndexFromAnyShader)
 		{
-			bIsInstancedStereoEnabled = true;
-		}
-		else
-		{
-			bIsMobileMultiViewEnabled = false;
+			UE_LOG(LogMultiView, Fatal, TEXT("Mobile Multi-View not supported by the RHI and no fallback is available."));
 		}
+		bIsInstancedStereoEnabled = RHISupportsInstancedStereo(ShaderPlatform);
 	}
 
 	// If the plugin uses separate render targets it is required to support mobile multi-view direct
diff --git a/Engine/Source/Runtime/Engine/Private/ShaderCompiler/ShaderCompiler.cpp b/Engine/Source/Runtime/Engine/Private/ShaderCompiler/ShaderCompiler.cpp
--- a/Engine/Source/Runtime/Engine/Private/ShaderCompiler/ShaderCompiler.cpp
+++ b/Engine/Source/Runtime/Engine/Private/ShaderCompiler/ShaderCompiler.cpp
@@ -3458,15 +3458,8 @@ void GlobalBeginCompileShader(
 		bool bIsMobileMultiview = bUsingMobileRenderer && bIsMobileMultiViewCVar;
 		if (bIsMobileMultiview && !RHISupportsMobileMultiView(ShaderPlatform))
 		{
-			// Native mobile multi-view is not supported, attempt to fall back to instancing
-			if (GRHISupportsArrayIndexFromAnyShader)
-			{
-				bIsInstancedStereo = true;
-			}
-			else
-			{
-				bIsMobileMultiview = false;
-			}
+			// Native mobile multi-view is not supported, fall back to instancing if available
+			bIsInstancedStereo = RHISupportsInstancedStereo(ShaderPlatform);
 		}
 
 		Input.Environment.SetDefine(TEXT("INSTANCED_STEREO"), bIsInstancedStereo);
diff --git a/Engine/Source/Runtime/Engine/Public/SceneManagement.h b/Engine/Source/Runtime/Engine/Public/SceneManagement.h
--- a/Engine/Source/Runtime/Engine/Public/SceneManagement.h
+++ b/Engine/Source/Runtime/Engine/Public/SceneManagement.h
@@ -53,6 +53,7 @@ struct FDynamicMeshVertex;
 class ULightMapVirtualTexture2D;
 
 DECLARE_LOG_CATEGORY_EXTERN(LogBufferVisualization, Log, All);
+DECLARE_LOG_CATEGORY_EXTERN(LogMultiView, Log, All);
 
 // -----------------------------------------------------------------------------
 
