PR #6861: HoloLens - Fall back to DX11 Viewport when using MultiView (Contributed by [REDACTED_PERSON])

diff --git a/Engine/Plugins/Runtime/WindowsMixedReality/Source/WindowsMixedRealityHMD/Private/WindowsMixedRealityHMD.cpp b/Engine/Plugins/Runtime/WindowsMixedReality/Source/WindowsMixedRealityHMD/Private/WindowsMixedRealityHMD.cpp
--- a/Engine/Plugins/Runtime/WindowsMixedReality/Source/WindowsMixedRealityHMD/Private/WindowsMixedRealityHMD.cpp
+++ b/Engine/Plugins/Runtime/WindowsMixedReality/Source/WindowsMixedRealityHMD/Private/WindowsMixedRealityHMD.cpp
@@ -1368,8 +1368,7 @@ namespace WindowsMixedReality
 		FRHIResourceCreateInfo CreateInfo;
 
 #if PLATFORM_HOLOLENS
-        // TODO: Remove index < 16 condition once FWindowsMixedRealityViewport has been removed
-		if (index < 16 && bIsMobileMultiViewEnabled)
+		if (bIsMobileMultiViewEnabled)
 		{
 			FTexture2DArrayRHIRef texture, resource;
 			RHICreateTargetableShaderResource2DArray(
diff --git a/Engine/Plugins/Runtime/WindowsMixedReality/Source/WindowsMixedRealityRHI/Private/WindowsMixedRealityRHI.cpp b/Engine/Plugins/Runtime/WindowsMixedReality/Source/WindowsMixedRealityRHI/Private/WindowsMixedRealityRHI.cpp
--- a/Engine/Plugins/Runtime/WindowsMixedReality/Source/WindowsMixedRealityRHI/Private/WindowsMixedRealityRHI.cpp
+++ b/Engine/Plugins/Runtime/WindowsMixedReality/Source/WindowsMixedRealityRHI/Private/WindowsMixedRealityRHI.cpp
@@ -166,6 +166,13 @@ FViewportRHIRef FWindowsMixedRealityDynamicRHI::RHICreateViewport(void* WindowHa
 {
 	check(IsInGameThread());
 
+#if PLATFORM_HOLOLENS
+	static const auto CVarMobileMultiView = IConsoleManager::Get().FindTConsoleVariableDataInt(TEXT("vr.MobileMultiView"));
+	static const auto CVarMobileHDR = IConsoleManager::Get().FindTConsoleVariableDataInt(TEXT("r.MobileHDR"));
+	const bool bMobileHDR = (CVarMobileHDR && CVarMobileHDR->GetValueOnAnyThread() != 0);
+	bIsMobileMultiViewEnabled = (GRHISupportsArrayIndexFromAnyShader && !bMobileHDR) && (CVarMobileMultiView && CVarMobileMultiView->GetValueOnAnyThread() != 0);
+#endif
+	
 	// Use a default pixel format if none was specified	
 	if (PreferredPixelFormat == EPixelFormat::PF_Unknown)
 	{
@@ -181,12 +188,23 @@ FViewportRHIRef FWindowsMixedRealityDynamicRHI::RHICreateViewport(void* WindowHa
 		IConsoleManager::Get().FindConsoleVariable(TEXT("RHI.MaximumFrameLatency"))->Set(16);
 	}
 
+	if (bIsMobileMultiViewEnabled)
+	{
+		return new FD3D11Viewport(this, (HWND)WindowHandle, SizeX, SizeY, bIsFullscreen, PreferredPixelFormat);
+	}
+	
 	return new FWindowsMixedRealityViewport(this, (HWND)WindowHandle, SizeX, SizeY, bIsFullscreen, PreferredPixelFormat);
 }
 
 void FWindowsMixedRealityDynamicRHI::RHIBeginFrame()
 {
 	FD3D11DynamicRHI::RHIBeginFrame();
+
+	if (bIsMobileMultiViewEnabled)
+	{
+		return;
+	}
+	
 	for (int i=0; i< Viewports.Num(); ++i)
 	{
 		((FWindowsMixedRealityViewport *)Viewports[i])->UpdateBackBuffer();
@@ -267,11 +285,10 @@ void FWindowsMixedRealityViewport::UpdateBackBuffer()
 	BackBuffer = nullptr;
 	if (GEngine && GEngine->StereoRenderingDevice && GEngine->StereoRenderingDevice->GetRenderTargetManager())
 	{
-        // HACK: Allocate with index 16 to avoid layered RTs from being allocated for viewports that don't need it
 		if (GEngine
 			->StereoRenderingDevice
 			->GetRenderTargetManager()
-			->AllocateRenderTargetTexture(16, SizeX, SizeY, PixelFormat, 1, TexCreate_None, TexCreate_RenderTargetable, OutTargetableTexture, OutShaderResourceTexture))
+			->AllocateRenderTargetTexture(0, SizeX, SizeY, PixelFormat, 1, TexCreate_None, TexCreate_RenderTargetable, OutTargetableTexture, OutShaderResourceTexture))
 		{
 			BackBuffer = (FD3D11Texture2D*)(OutTargetableTexture.GetReference());
 		}
diff --git a/Engine/Plugins/Runtime/WindowsMixedReality/Source/WindowsMixedRealityRHI/Private/WindowsMixedRealityRHI.h b/Engine/Plugins/Runtime/WindowsMixedReality/Source/WindowsMixedRealityRHI/Private/WindowsMixedRealityRHI.h
--- a/Engine/Plugins/Runtime/WindowsMixedReality/Source/WindowsMixedRealityRHI/Private/WindowsMixedRealityRHI.h
+++ b/Engine/Plugins/Runtime/WindowsMixedReality/Source/WindowsMixedRealityRHI/Private/WindowsMixedRealityRHI.h
@@ -63,6 +63,7 @@ public:
 	HolographicSpace^ GetHolographicSpace() const { return HoloSpace; }
 private:
 	HolographicSpace^ HoloSpace;
+	bool bIsMobileMultiViewEnabled = false;
 #endif // PLATFORM_HOLOLENS
 	
 };
@@ -78,4 +79,4 @@ public:
 
 private:
 	TRefCountPtr<FD3D11Texture2D> OffscreenBackBuffer;
-};
\ No newline at end of file
+};
