Remote beta tags on WMR and HoloLens plugins.

diff --git a/Engine/Plugins/Runtime/AR/Microsoft/HoloLensAR/HoloLensAR.uplugin b/Engine/Plugins/Runtime/AR/Microsoft/HoloLensAR/HoloLensAR.uplugin
--- a/Engine/Plugins/Runtime/AR/Microsoft/HoloLensAR/HoloLensAR.uplugin
+++ b/Engine/Plugins/Runtime/AR/Microsoft/HoloLensAR/HoloLensAR.uplugin
@@ -12,7 +12,7 @@
 	"SupportURL": "",
 	"EnabledByDefault": false,
 	"CanContainContent": true,
-	"IsBetaVersion": true,
+	"IsBetaVersion": false,
 	"Installed": false,
 	"SupportedTargetPlatforms": [
 		"HoloLens",
diff --git a/Engine/Plugins/Runtime/WindowsMixedReality/WindowsMixedReality.uplugin b/Engine/Plugins/Runtime/WindowsMixedReality/WindowsMixedReality.uplugin
--- a/Engine/Plugins/Runtime/WindowsMixedReality/WindowsMixedReality.uplugin
+++ b/Engine/Plugins/Runtime/WindowsMixedReality/WindowsMixedReality.uplugin
@@ -12,7 +12,7 @@
 	"SupportURL" : "",
 	"EnabledByDefault" : false,
 	"CanContainContent" : false,
-	"IsBetaVersion" : true,
+	"IsBetaVersion" : false,
 	"Installed" : false,
 	"SupportedTargetPlatforms": [
 		"Win64",
