Build fix for Compile errors in WindowsMixedRealityHMD.cpp and Platform.h
-Fixed wrong parameter removed.

diff --git a/Engine/Plugins/Runtime/WindowsMixedReality/Source/WindowsMixedRealityHMD/Private/WindowsMixedRealityHMD.cpp b/Engine/Plugins/Runtime/WindowsMixedReality/Source/WindowsMixedRealityHMD/Private/WindowsMixedRealityHMD.cpp
--- a/Engine/Plugins/Runtime/WindowsMixedReality/Source/WindowsMixedRealityHMD/Private/WindowsMixedRealityHMD.cpp
+++ b/Engine/Plugins/Runtime/WindowsMixedReality/Source/WindowsMixedRealityHMD/Private/WindowsMixedRealityHMD.cpp
@@ -1293,19 +1293,19 @@ namespace WindowsMixedReality
 			{
 				if (ParseAddress(MatchesStr, IPAddress, Port))
 				{
-					ConnectToRemoteHoloLens(*IPAddress, BitRate, false, Port);
+					ConnectToRemoteHoloLens(*IPAddress, BitRate, Port);
 				}
 			}
 			else if (FParse::Value(FCommandLine::Get(), TEXT("HoloLensRemotingListen="), MatchesStr))
 			{
 				if (ParseAddress(MatchesStr, IPAddress, Port))
 				{
-					ConnectToRemoteHoloLens(*IPAddress, BitRate, false, Port);
+					ConnectToRemoteHoloLens(*IPAddress, BitRate, Port, true);
 				}
 			}
 			else if (FParse::Value(FCommandLine::Get(), TEXT("HoloLensRemotingListenPort="), Port))
 			{
-				ConnectToRemoteHoloLens(TEXT("0.0.0.0"), BitRate, false, Port);
+				ConnectToRemoteHoloLens(TEXT("0.0.0.0"), BitRate, Port, true);
 			}
 #endif
 			HMD->EnableStereo(stereo);
