Horde HolographicStreamerDesktop.dll does not exist (and similar)
-Removed references  to deleted Hololens1 dlls.

diff --git a/Engine/Plugins/Runtime/MicrosoftSpatialAudio/Source/MicrosoftSpatialSound/Private/MicrosoftSpatialSoundPlugin.cpp b/Engine/Plugins/Runtime/MicrosoftSpatialAudio/Source/MicrosoftSpatialSound/Private/MicrosoftSpatialSoundPlugin.cpp
--- a/Engine/Plugins/Runtime/MicrosoftSpatialAudio/Source/MicrosoftSpatialSound/Private/MicrosoftSpatialSoundPlugin.cpp
+++ b/Engine/Plugins/Runtime/MicrosoftSpatialAudio/Source/MicrosoftSpatialSound/Private/MicrosoftSpatialSoundPlugin.cpp
@@ -437,9 +437,6 @@ void FMicrosoftSpatialSoundModule::StartupModule()
 	FPlatformProcess::GetDllHandle(_TEXT("PerceptionDevice.dll"));
 	FPlatformProcess::GetDllHandle(_TEXT("Microsoft.Holographic.AppRemoting.dll"));
 	FPlatformProcess::PopDllDirectory(*HoloLensLibraryDir);
-
-	FPlatformProcess::GetDllHandle(*(EngineDir / "Binaries" / BinariesSubDir / "HolographicStreamerDesktop.dll"));
-	FPlatformProcess::GetDllHandle(*(EngineDir / "Binaries" / BinariesSubDir / "Microsoft.Perception.Simulation.dll"));
 #endif // PLATFORM_64BITS && WITH_EDITOR	
 	
 	// Then finally try to load the WMR Interop Library
diff --git a/Engine/Plugins/Runtime/WindowsMixedReality/Source/WindowsMixedRealityHMD/Private/WindowsMixedRealityInteropLoader.cpp b/Engine/Plugins/Runtime/WindowsMixedReality/Source/WindowsMixedRealityHMD/Private/WindowsMixedRealityInteropLoader.cpp
--- a/Engine/Plugins/Runtime/WindowsMixedReality/Source/WindowsMixedRealityHMD/Private/WindowsMixedRealityInteropLoader.cpp
+++ b/Engine/Plugins/Runtime/WindowsMixedReality/Source/WindowsMixedRealityHMD/Private/WindowsMixedRealityInteropLoader.cpp
@@ -72,9 +72,6 @@ namespace WindowsMixedReality
 			FPlatformProcess::GetDllHandle(_TEXT("Microsoft.MixedReality.SceneUnderstanding.dll"));
 			FPlatformProcess::GetDllHandle(_TEXT("Microsoft.Azure.SpatialAnchors.dll"));
 			FPlatformProcess::PopDllDirectory(*HoloLensLibraryDir);
-
-			FPlatformProcess::GetDllHandle(*(EngineDir / "Binaries" / BinariesSubDir / "HolographicStreamerDesktop.dll"));
-			FPlatformProcess::GetDllHandle(*(EngineDir / "Binaries" / BinariesSubDir / "Microsoft.Perception.Simulation.dll"));
 #endif // PLATFORM_64BITS && WITH_EDITOR
 
 			// Then finally try to load the WMR Interop Library
diff --git a/Engine/Plugins/Runtime/WindowsMixedReality/Source/WindowsMixedRealityHMD/WindowsMixedRealityHMD.Build.cs b/Engine/Plugins/Runtime/WindowsMixedReality/Source/WindowsMixedRealityHMD/WindowsMixedRealityHMD.Build.cs
--- a/Engine/Plugins/Runtime/WindowsMixedReality/Source/WindowsMixedRealityHMD/WindowsMixedRealityHMD.Build.cs
+++ b/Engine/Plugins/Runtime/WindowsMixedReality/Source/WindowsMixedRealityHMD/WindowsMixedRealityHMD.Build.cs
@@ -34,11 +34,6 @@ namespace UnrealBuildTool.Rules
 				// HoloLens 2 Remoting
                 RuntimeDependencies.Add("$(EngineDir)/Binaries/ThirdParty/Windows/x64/Microsoft.Holographic.AppRemoting.dll");
                 RuntimeDependencies.Add("$(EngineDir)/Binaries/ThirdParty/Windows/x64/PerceptionDevice.dll");
-				
-				// HoloLens 1 Remoting
-                RuntimeDependencies.Add("$(EngineDir)/Binaries/Win64/HolographicStreamerDesktop.dll");
-                RuntimeDependencies.Add("$(EngineDir)/Binaries/Win64/Microsoft.Perception.Simulation.dll");
-                RuntimeDependencies.Add("$(EngineDir)/Binaries/Win64/PerceptionSimulationManager.dll");
             }
 
             PublicDefinitions.Add("WITH_WINDOWS_MIXED_REALITY=1");
