Fix for WMR not depending on RHICore

diff --git a/Engine/Plugins/Runtime/WindowsMixedReality/Source/WindowsMixedRealityRHI/WindowsMixedRealityRHI.Build.cs b/Engine/Plugins/Runtime/WindowsMixedReality/Source/WindowsMixedRealityRHI/WindowsMixedRealityRHI.Build.cs
--- a/Engine/Plugins/Runtime/WindowsMixedReality/Source/WindowsMixedRealityRHI/WindowsMixedRealityRHI.Build.cs
+++ b/Engine/Plugins/Runtime/WindowsMixedReality/Source/WindowsMixedRealityRHI/WindowsMixedRealityRHI.Build.cs
@@ -54,6 +54,7 @@ namespace UnrealBuildTool.Rules
 						"Projects",
                         "HeadMountedDisplay",
                         "D3D11RHI",
+						"RHICore",
                         "Slate",
                         "SlateCore",
                         "MRMesh",
