WindowsMR: Configure GShaderPlatformForFeatureLevel at RHI creation time

diff --git a/Engine/Plugins/Runtime/WindowsMixedReality/Source/WindowsMixedRealityRHI/Private/WindowsMixedRealityRHI.cpp b/Engine/Plugins/Runtime/WindowsMixedReality/Source/WindowsMixedRealityRHI/Private/WindowsMixedRealityRHI.cpp
--- a/Engine/Plugins/Runtime/WindowsMixedReality/Source/WindowsMixedRealityRHI/Private/WindowsMixedRealityRHI.cpp
+++ b/Engine/Plugins/Runtime/WindowsMixedReality/Source/WindowsMixedRealityRHI/Private/WindowsMixedRealityRHI.cpp
@@ -149,6 +149,11 @@ FDynamicRHI* FWindowsMixedRealityRHIModule::CreateRHI(ERHIFeatureLevel::Type Req
 	GMaxRHIShaderPlatform = SP_PCD3D_ES3_1;
 #endif
 
+	GShaderPlatformForFeatureLevel[ERHIFeatureLevel::ES2_REMOVED] = SP_NumPlatforms;
+	GShaderPlatformForFeatureLevel[ERHIFeatureLevel::ES3_1] = SP_PCD3D_ES3_1;
+	GShaderPlatformForFeatureLevel[ERHIFeatureLevel::SM4_REMOVED] = SP_NumPlatforms;
+	GShaderPlatformForFeatureLevel[ERHIFeatureLevel::SM5] = SP_PCD3D_SM5;
+
 	FWindowsMixedRealityDynamicRHI * RHI = new FWindowsMixedRealityDynamicRHI(DXGIFactory1.Get(), ChosenAdapter.MaxSupportedFeatureLevel, ChosenAdapter.AdapterIndex, ChosenDescription);
 	GD3D11RHI = RHI;
 
