Speculative fix for clang build error about INTEL_EXTENSIONS not being defined.

#rnx

diff --git a/Engine/Plugins/Runtime/WindowsMixedReality/Source/WindowsMixedRealityHMD/WindowsMixedRealityHMD.Build.cs b/Engine/Plugins/Runtime/WindowsMixedReality/Source/WindowsMixedRealityHMD/WindowsMixedRealityHMD.Build.cs
--- a/Engine/Plugins/Runtime/WindowsMixedReality/Source/WindowsMixedRealityHMD/WindowsMixedRealityHMD.Build.cs
+++ b/Engine/Plugins/Runtime/WindowsMixedReality/Source/WindowsMixedRealityHMD/WindowsMixedRealityHMD.Build.cs
@@ -95,6 +95,7 @@ namespace UnrealBuildTool.Rules
 				{
 					AddEngineThirdPartyPrivateStaticDependencies(Target, "NVAftermath");
 					AddEngineThirdPartyPrivateStaticDependencies(Target, "IntelMetricsDiscovery");
+					AddEngineThirdPartyPrivateStaticDependencies(Target, "IntelExtensionsFramework");
 				}
 
                 AddEngineThirdPartyPrivateStaticDependencies(Target, "WindowsMixedRealityInterop");
diff --git a/Engine/Plugins/Runtime/WindowsMixedReality/Source/WindowsMixedRealityRHI/WindowsMixedRealityRHI.Build.cs b/Engine/Plugins/Runtime/WindowsMixedReality/Source/WindowsMixedRealityRHI/WindowsMixedRealityRHI.Build.cs
--- a/Engine/Plugins/Runtime/WindowsMixedReality/Source/WindowsMixedRealityRHI/WindowsMixedRealityRHI.Build.cs
+++ b/Engine/Plugins/Runtime/WindowsMixedReality/Source/WindowsMixedRealityRHI/WindowsMixedRealityRHI.Build.cs
@@ -63,7 +63,13 @@ namespace UnrealBuildTool.Rules
 
                 AddEngineThirdPartyPrivateStaticDependencies(Target, "WindowsMixedRealityInterop");
 
-                if (Target.bBuildEditor == true)
+				if (Target.Platform != UnrealTargetPlatform.HoloLens)
+				{
+					AddEngineThirdPartyPrivateStaticDependencies(Target, "NVAftermath");
+					AddEngineThirdPartyPrivateStaticDependencies(Target, "IntelExtensionsFramework");
+				}
+
+				if (Target.bBuildEditor == true)
                 {
 					PrivateDependencyModuleNames.Add("EditorFramework");
                     PrivateDependencyModuleNames.Add("UnrealEd");
