More TStaticArray deprecation fixes.

diff --git a/Engine/Plugins/Runtime/WindowsMixedReality/Source/WindowsMixedRealityInputSimulation/Public/WindowsMixedRealityInputSimulationEngineSubsystem.h b/Engine/Plugins/Runtime/WindowsMixedReality/Source/WindowsMixedRealityInputSimulation/Public/WindowsMixedRealityInputSimulationEngineSubsystem.h
--- a/Engine/Plugins/Runtime/WindowsMixedReality/Source/WindowsMixedRealityInputSimulation/Public/WindowsMixedRealityInputSimulationEngineSubsystem.h
+++ b/Engine/Plugins/Runtime/WindowsMixedReality/Source/WindowsMixedRealityInputSimulation/Public/WindowsMixedRealityInputSimulationEngineSubsystem.h
@@ -31,8 +31,8 @@ struct WINDOWSMIXEDREALITYINPUTSIMULATION_API FWindowsMixedRealityInputSimulatio
 	ETrackingStatus TrackingStatus = ETrackingStatus::NotTracked;
 
 	bool bHasJointPoses = false;
-	KeypointTransformArray KeypointTransforms = KeypointTransformArray(FTransform::Identity);
-	KeypointRadiusArray KeypointRadii = KeypointRadiusArray(1.0f);
+	KeypointTransformArray KeypointTransforms = KeypointTransformArray(InPlace, FTransform::Identity);
+	KeypointRadiusArray KeypointRadii = KeypointRadiusArray(InPlace, 1.0f);
 
 	ButtonStateArray IsButtonPressed = ButtonStateArray();
 	ButtonStateArray PrevButtonPressed = ButtonStateArray();
