Fixing CIS No Unity issues

diff --git a/Engine/Plugins/Experimental/EnhancedInput/Source/EnhancedInput/Private/EnhancedPlayerInput.cpp b/Engine/Plugins/Experimental/EnhancedInput/Source/EnhancedInput/Private/EnhancedPlayerInput.cpp
--- a/Engine/Plugins/Experimental/EnhancedInput/Source/EnhancedInput/Private/EnhancedPlayerInput.cpp
+++ b/Engine/Plugins/Experimental/EnhancedInput/Source/EnhancedInput/Private/EnhancedPlayerInput.cpp
@@ -2,6 +2,7 @@
 
 #include "EnhancedPlayerInput.h"
 #include "EnhancedInputComponent.h"
+#include "GameFramework/PlayerController.h"
 #include "InputModifiers.h"
 #include "InputTriggers.h"
 
diff --git a/Engine/Plugins/Experimental/EnhancedInput/Source/EnhancedInput/Private/InputTriggers.cpp b/Engine/Plugins/Experimental/EnhancedInput/Source/EnhancedInput/Private/InputTriggers.cpp
--- a/Engine/Plugins/Experimental/EnhancedInput/Source/EnhancedInput/Private/InputTriggers.cpp
+++ b/Engine/Plugins/Experimental/EnhancedInput/Source/EnhancedInput/Private/InputTriggers.cpp
@@ -4,7 +4,7 @@
 #include "InputAction.h"
 #include "InputActionValue.h"
 #include "EnhancedPlayerInput.h"
-
+#include "GameFramework/PlayerController.h"
 
 // Abstract trigger bases
 ETriggerState UInputTrigger::UpdateState_Implementation(const UEnhancedPlayerInput* PlayerInput, FInputActionValue ModifiedValue, float DeltaTime)
diff --git a/Engine/Plugins/Experimental/GameFeatures/Source/GameFeatures/Private/GameFeaturePluginStateMachine.cpp b/Engine/Plugins/Experimental/GameFeatures/Source/GameFeatures/Private/GameFeaturePluginStateMachine.cpp
--- a/Engine/Plugins/Experimental/GameFeatures/Source/GameFeatures/Private/GameFeaturePluginStateMachine.cpp
+++ b/Engine/Plugins/Experimental/GameFeatures/Source/GameFeatures/Private/GameFeaturePluginStateMachine.cpp
@@ -16,6 +16,7 @@
 #include "Misc/Paths.h"
 #include "Algo/AllOf.h"
 #include "Serialization/ArrayReader.h"
+#include "Serialization/MemoryReader.h"
 #include "Stats/Stats.h"
 #include "EngineUtils.h"
 #include "GameFeaturesSubsystem.h"
diff --git a/Engine/Plugins/FX/Niagara/Source/Niagara/Private/NiagaraRenderViewDataManager.cpp b/Engine/Plugins/FX/Niagara/Source/Niagara/Private/NiagaraRenderViewDataManager.cpp
--- a/Engine/Plugins/FX/Niagara/Source/Niagara/Private/NiagaraRenderViewDataManager.cpp
+++ b/Engine/Plugins/FX/Niagara/Source/Niagara/Private/NiagaraRenderViewDataManager.cpp
@@ -3,6 +3,7 @@
 #include "NiagaraRenderViewDataManager.h"
 
 #include "EngineModule.h"
+#include "SceneInterface.h"
 #include "SceneRenderTargetParameters.h"
 
 TGlobalResource<FNiagaraRenderViewDataManager> GNiagaraViewDataManager;
diff --git a/Engine/Plugins/Runtime/NetworkPredictionExtras/Source/NetworkPredictionExtras/Private/MockNetworkPhysics.cpp b/Engine/Plugins/Runtime/NetworkPredictionExtras/Source/NetworkPredictionExtras/Private/MockNetworkPhysics.cpp
--- a/Engine/Plugins/Runtime/NetworkPredictionExtras/Source/NetworkPredictionExtras/Private/MockNetworkPhysics.cpp
+++ b/Engine/Plugins/Runtime/NetworkPredictionExtras/Source/NetworkPredictionExtras/Private/MockNetworkPhysics.cpp
@@ -87,6 +87,7 @@ FMockState_PT
 =============================================================================*/
 
 #include "MockNetworkPhysics.h"
+#include "Components/PrimitiveComponent.h"
 #include "Engine/World.h"
 #include "Physics/Experimental/PhysScene_Chaos.h"
 #include "PBDRigidsSolver.h"
diff --git a/Engine/Plugins/Runtime/Steam/SteamSockets/Source/SteamSockets/Private/SteamSocketsSubsystem.cpp b/Engine/Plugins/Runtime/Steam/SteamSockets/Source/SteamSockets/Private/SteamSocketsSubsystem.cpp
--- a/Engine/Plugins/Runtime/Steam/SteamSockets/Source/SteamSockets/Private/SteamSocketsSubsystem.cpp
+++ b/Engine/Plugins/Runtime/Steam/SteamSockets/Source/SteamSockets/Private/SteamSocketsSubsystem.cpp
@@ -1,6 +1,7 @@
 // Copyright Epic Games, Inc. All Rights Reserved.
 
 #include "SteamSocketsSubsystem.h"
+#include "Misc/CommandLine.h"
 #include "Misc/CoreMisc.h"
 #include "Misc/ConfigCacheIni.h"
 #include "IPAddressSteamSockets.h"
diff --git a/Engine/Plugins/Runtime/WindowsMixedReality/Source/WindowsMixedRealityEyeTracker/Private/WindowsMixedRealityEyeTracker.cpp b/Engine/Plugins/Runtime/WindowsMixedReality/Source/WindowsMixedRealityEyeTracker/Private/WindowsMixedRealityEyeTracker.cpp
--- a/Engine/Plugins/Runtime/WindowsMixedReality/Source/WindowsMixedRealityEyeTracker/Private/WindowsMixedRealityEyeTracker.cpp
+++ b/Engine/Plugins/Runtime/WindowsMixedReality/Source/WindowsMixedRealityEyeTracker/Private/WindowsMixedRealityEyeTracker.cpp
@@ -3,6 +3,7 @@
 #include "IEyeTracker.h"
 #include "Modules/ModuleManager.h"
 #include "GameFramework/WorldSettings.h"
+#include "Engine/Engine.h"
 #include "Engine/World.h"
 #include "WindowsMixedRealityInteropLoader.h"
 #include "IXRTrackingSystem.h"
diff --git a/Engine/Plugins/VirtualProduction/LevelSnapshots/Source/LevelSnapshots/Private/Data/LevelSnapshot.cpp b/Engine/Plugins/VirtualProduction/LevelSnapshots/Source/LevelSnapshots/Private/Data/LevelSnapshot.cpp
--- a/Engine/Plugins/VirtualProduction/LevelSnapshots/Source/LevelSnapshots/Private/Data/LevelSnapshot.cpp
+++ b/Engine/Plugins/VirtualProduction/LevelSnapshots/Source/LevelSnapshots/Private/Data/LevelSnapshot.cpp
@@ -11,6 +11,7 @@
 
 #include "Algo/Accumulate.h"
 #include "EngineUtils.h"
+#include "Engine/Engine.h"
 #include "GameFramework/Actor.h"
 #include "SnapshotRestorability.h"
 #include "UObject/Package.h"
diff --git a/Engine/Plugins/VirtualProduction/RemoteControl/Source/WebRemoteControl/Private/WebRemoteControl.h b/Engine/Plugins/VirtualProduction/RemoteControl/Source/WebRemoteControl/Private/WebRemoteControl.h
--- a/Engine/Plugins/VirtualProduction/RemoteControl/Source/WebRemoteControl/Private/WebRemoteControl.h
+++ b/Engine/Plugins/VirtualProduction/RemoteControl/Source/WebRemoteControl/Private/WebRemoteControl.h
@@ -15,6 +15,7 @@
 #include "WebRemoteControlEditorRoutes.h"
 
 struct FHttpServerRequest;
+class FMemoryWriter;
 class IHttpRouter;
 class FWebSocketMessageHandler;
 
