Fix naming in WindowsMixedRealityInteropLoader (UE4 -> UE).

diff --git a/Engine/Plugins/Runtime/WindowsMixedReality/Source/WindowsMixedRealityHMD/Private/WindowsMixedRealityInteropLoader.cpp b/Engine/Plugins/Runtime/WindowsMixedReality/Source/WindowsMixedRealityHMD/Private/WindowsMixedRealityInteropLoader.cpp
--- a/Engine/Plugins/Runtime/WindowsMixedReality/Source/WindowsMixedRealityHMD/Private/WindowsMixedRealityInteropLoader.cpp
+++ b/Engine/Plugins/Runtime/WindowsMixedReality/Source/WindowsMixedRealityHMD/Private/WindowsMixedRealityInteropLoader.cpp
@@ -89,7 +89,7 @@ namespace WindowsMixedReality
 		}
 		else
 		{
-			FText ErrorText = FText::Format(FTextFormat(NSLOCTEXT("WindowsMixedRealityHMD", "MixedRealityInteropLibraryNotSupported", "Windows Mixed Reality is not supported on this Windows version. \nNote: UE4 only supports Windows Mixed Reality on Windows 10 Release {0} or higher. Current version: {1}")),
+			FText ErrorText = FText::Format(FTextFormat(NSLOCTEXT("WindowsMixedRealityHMD", "MixedRealityInteropLibraryNotSupported", "Windows Mixed Reality is not supported on this Windows version. \nNote: UE only supports Windows Mixed Reality on Windows 10 Release {0} or higher. Current version: {1}")),
 				FText::FromString(FString::FromInt(MIN_WIN_10_VERSION_FOR_WMR)), FText::FromString(OSVersionLabel));
 			FMessageDialog::Open(EAppMsgType::Ok, ErrorText);
 			if (IsRunningCommandlet())
