Fix UE4 references in WindowsMixedRealityHMD.cpp (UE4 -> UE)

diff --git a/Engine/Plugins/Runtime/WindowsMixedReality/Source/WindowsMixedRealityHMD/Private/WindowsMixedRealityHMD.cpp b/Engine/Plugins/Runtime/WindowsMixedReality/Source/WindowsMixedRealityHMD/Private/WindowsMixedRealityHMD.cpp
--- a/Engine/Plugins/Runtime/WindowsMixedReality/Source/WindowsMixedRealityHMD/Private/WindowsMixedRealityHMD.cpp
+++ b/Engine/Plugins/Runtime/WindowsMixedReality/Source/WindowsMixedRealityHMD/Private/WindowsMixedRealityHMD.cpp
@@ -2335,7 +2335,7 @@ namespace WindowsMixedReality
 		}
 	}
 
-	// GetControllerOrientationAndPosition() is pre-scaled to UE4 world scale by default
+	// GetControllerOrientationAndPosition() is pre-scaled to UE world scale by default
 	bool FWindowsMixedRealityHMD::GetControllerOrientationAndPosition(HMDHand hand, FRotator & OutOrientation, FVector & OutPosition, float WorldScale)
 	{
 		if (!bIsStereoEnabled)
@@ -2366,7 +2366,7 @@ namespace WindowsMixedReality
 		return false;
 	}
 
-	// GetHandJointOrientationAndPosition() is NOT scaled to UE4 world scale, so we must do it here
+	// GetHandJointOrientationAndPosition() is NOT scaled to UE world scale, so we must do it here
 	bool FWindowsMixedRealityHMD::GetHandJointOrientationAndPosition(HMDHand hand, HMDHandJoint joint, FRotator& OutOrientation, FVector& OutPosition, float& OutRadius)
 	{
 		if (!bIsStereoEnabled)
