UE-125001 //UE5/Main - Compile UnrealGame HoloLens - Several compile errors
-Fixed One WindowsMixedRealityRHI.cpp build break.

diff --git a/Engine/Plugins/Runtime/WindowsMixedReality/Source/WindowsMixedRealityRHI/Private/WindowsMixedRealityRHI.cpp b/Engine/Plugins/Runtime/WindowsMixedReality/Source/WindowsMixedRealityRHI/Private/WindowsMixedRealityRHI.cpp
--- a/Engine/Plugins/Runtime/WindowsMixedReality/Source/WindowsMixedRealityRHI/Private/WindowsMixedRealityRHI.cpp
+++ b/Engine/Plugins/Runtime/WindowsMixedReality/Source/WindowsMixedRealityRHI/Private/WindowsMixedRealityRHI.cpp
@@ -281,8 +281,12 @@ FWindowsMixedRealityViewport::FWindowsMixedRealityViewport(FD3D11DynamicRHI* InD
 	OffscreenBackBuffer = FD3D11Viewport::GetSwapChainSurface(D3DRHI, PixelFormat, SizeX, SizeY, nullptr);
 	UpdateBackBuffer();
 
-	BeginInitResource(&FrameSyncEvent);
-
+	ENQUEUE_RENDER_COMMAND(FD3D11Viewport)(
+	[this](FRHICommandListImmediate& RHICmdList)
+	{
+		// Initialize the query by issuing an initial event.
+		FrameSyncEvent.IssueEvent();
+	});
 }
 
 void FWindowsMixedRealityViewport::UpdateBackBuffer()
