Fix build break in WMR D3D11 RHI implementation.

diff --git a/Engine/Plugins/Runtime/WindowsMixedReality/Source/WindowsMixedRealityRHI/Private/WindowsMixedRealityRHI.cpp b/Engine/Plugins/Runtime/WindowsMixedReality/Source/WindowsMixedRealityRHI/Private/WindowsMixedRealityRHI.cpp
--- a/Engine/Plugins/Runtime/WindowsMixedReality/Source/WindowsMixedRealityRHI/Private/WindowsMixedRealityRHI.cpp
+++ b/Engine/Plugins/Runtime/WindowsMixedReality/Source/WindowsMixedRealityRHI/Private/WindowsMixedRealityRHI.cpp
@@ -104,7 +104,6 @@ void FWindowsMixedRealityRHIModule::StartupModule()
 			continue; 
 		}
 
-		ChosenAdapter.AdapterIndex = AdapterIndex;
 		ChosenAdapter.MaxSupportedFeatureLevel = OutFeatureLevel;
 		ChosenDescription = AdapterDesc;
 		break;
@@ -154,7 +153,7 @@ FDynamicRHI* FWindowsMixedRealityRHIModule::CreateRHI(ERHIFeatureLevel::Type Req
 	GShaderPlatformForFeatureLevel[ERHIFeatureLevel::SM4_REMOVED] = SP_NumPlatforms;
 	GShaderPlatformForFeatureLevel[ERHIFeatureLevel::SM5] = SP_PCD3D_SM5;
 
-	FWindowsMixedRealityDynamicRHI * RHI = new FWindowsMixedRealityDynamicRHI(DXGIFactory1.Get(), ChosenAdapter.MaxSupportedFeatureLevel, ChosenAdapter.AdapterIndex, ChosenDescription);
+	FWindowsMixedRealityDynamicRHI * RHI = new FWindowsMixedRealityDynamicRHI(DXGIFactory1.Get(), ChosenAdapter.MaxSupportedFeatureLevel, ChosenAdapter);
 	GD3D11RHI = RHI;
 
 #if PLATFORM_HOLOLENS
@@ -165,8 +164,8 @@ FDynamicRHI* FWindowsMixedRealityRHIModule::CreateRHI(ERHIFeatureLevel::Type Req
 }
 
 
-FWindowsMixedRealityDynamicRHI::FWindowsMixedRealityDynamicRHI(IDXGIFactory1* InDXGIFactory1, D3D_FEATURE_LEVEL InFeatureLevel, int32 InChosenAdapter, const DXGI_ADAPTER_DESC& InChosenDescription) 
-	: FD3D11DynamicRHI(InDXGIFactory1, InFeatureLevel, InChosenAdapter, InChosenDescription)
+FWindowsMixedRealityDynamicRHI::FWindowsMixedRealityDynamicRHI(IDXGIFactory1* InDXGIFactory1, D3D_FEATURE_LEVEL InFeatureLevel, const FD3D11Adapter& InAdapter)
+	: FD3D11DynamicRHI(InDXGIFactory1, InFeatureLevel, InAdapter)
 {}
 
 
diff --git a/Engine/Plugins/Runtime/WindowsMixedReality/Source/WindowsMixedRealityRHI/Private/WindowsMixedRealityRHI.h b/Engine/Plugins/Runtime/WindowsMixedReality/Source/WindowsMixedRealityRHI/Private/WindowsMixedRealityRHI.h
--- a/Engine/Plugins/Runtime/WindowsMixedReality/Source/WindowsMixedRealityRHI/Private/WindowsMixedRealityRHI.h
+++ b/Engine/Plugins/Runtime/WindowsMixedReality/Source/WindowsMixedRealityRHI/Private/WindowsMixedRealityRHI.h
@@ -50,7 +50,7 @@ private:
 class FWindowsMixedRealityDynamicRHI : public FD3D11DynamicRHI
 {
 public:
-	FWindowsMixedRealityDynamicRHI(IDXGIFactory1* InDXGIFactory1, D3D_FEATURE_LEVEL InFeatureLevel, int32 InChosenAdapter, const DXGI_ADAPTER_DESC& ChosenDescription);
+	FWindowsMixedRealityDynamicRHI(IDXGIFactory1* InDXGIFactory1, D3D_FEATURE_LEVEL InFeatureLevel, const FD3D11Adapter& InAdapter);
 
 	virtual FViewportRHIRef RHICreateViewport(void* WindowHandle, uint32 SizeX, uint32 SizeY, bool bIsFullscreen, EPixelFormat PreferredPixelFormat) override;
 
