PR #7681: Implementation of GetCameraIntrinsics for WMR (Contributed by [REDACTED_PERSON])
(Plus some coding standard edits)

diff --git a/Engine/Plugins/Runtime/WindowsMixedReality/Source/HoloLensAR/Private/HoloLensARSystem.cpp b/Engine/Plugins/Runtime/WindowsMixedReality/Source/HoloLensAR/Private/HoloLensARSystem.cpp
--- a/Engine/Plugins/Runtime/WindowsMixedReality/Source/HoloLensAR/Private/HoloLensARSystem.cpp
+++ b/Engine/Plugins/Runtime/WindowsMixedReality/Source/HoloLensAR/Private/HoloLensARSystem.cpp
@@ -811,45 +811,65 @@ FTransform FHoloLensARSystem::GetPVCameraToWorldTransform()
 	return PVCameraToWorldMatrix;
 }
 
-bool FHoloLensARSystem::GetPVCameraIntrinsics(FVector2D& focalLength, int& width, int& height, FVector2D& principalPoint, FVector& radialDistortion, FVector2D& tangentialDistortion)
+bool FHoloLensARSystem::GetPVCameraIntrinsics(FVector2D& OutFocalLength, int& OutWidth, int& OutHeight, FVector2D& OutPrincipalPoint, FVector& OutRadialDistortion, FVector2D& OutTangentialDistortion)
 {
 	CameraImageCapture& CameraCapture = CameraImageCapture::Get();
 	
-	DirectX::XMFLOAT2 _focalLength, _principalPoint, _tangentialDistortion;
-	DirectX::XMFLOAT3 _radialDistortion;
-	if (!CameraCapture.GetCameraIntrinsics(_focalLength, width, height, _principalPoint, _radialDistortion, _tangentialDistortion))
+	DirectX::XMFLOAT2 FocalLength, PrincipalPoint, TangentialDistortion;
+	DirectX::XMFLOAT3 RadialDistortion;
+	if (!CameraCapture.GetCameraIntrinsics(FocalLength, OutWidth, OutHeight, PrincipalPoint, RadialDistortion, TangentialDistortion))
 	{
 		return false;
 	}
 
 	// Convert to FVector - 2d Vectors preserve windows coordinate system (x is left/right, y is up/down)
-	focalLength = FVector2D(_focalLength.x, _focalLength.y);
-	principalPoint = FVector2D(_principalPoint.x, _principalPoint.y);
-	radialDistortion = WindowsMixedReality::WMRUtility::FromMixedRealityVector(_radialDistortion);
-	tangentialDistortion = FVector2D(_tangentialDistortion.x, _tangentialDistortion.y);
+	OutFocalLength = FVector2D(FocalLength.x, FocalLength.y);
+	OutPrincipalPoint = FVector2D(PrincipalPoint.x, PrincipalPoint.y);
+	OutRadialDistortion = WindowsMixedReality::WMRUtility::FromMixedRealityVector(RadialDistortion);
+	OutTangentialDistortion = FVector2D(TangentialDistortion.x, TangentialDistortion.y);
 
 	return true;
 }
 
-FVector FHoloLensARSystem::GetWorldSpaceRayFromCameraPoint(FVector2D pixelCoordinate)
+bool FHoloLensARSystem::OnGetCameraIntrinsics(FARCameraIntrinsics& OutCameraIntrinsics) const
+{
+	CameraImageCapture& CameraCapture = CameraImageCapture::Get();
+
+	DirectX::XMFLOAT2 FocalLength, PrincipalPoint, TangentialDistortion;
+	DirectX::XMFLOAT3 RadialDistortion;
+	int Width, Height;
+	if (!CameraCapture.GetCameraIntrinsics(FocalLength, Width, Height, PrincipalPoint, RadialDistortion, TangentialDistortion))
+	{
+		return false;
+	}
+
+	// Convert to FVector - 2d Vectors preserve windows coordinate system (x is left/right, y is up/down)
+	OutCameraIntrinsics.FocalLength = FVector2D(FocalLength.x, FocalLength.y);
+	OutCameraIntrinsics.PrincipalPoint = FVector2D(PrincipalPoint.x, PrincipalPoint.y);
+	OutCameraIntrinsics.ImageResolution = FIntPoint(Width, Height);
+
+	return true;
+}
+
+FVector FHoloLensARSystem::GetWorldSpaceRayFromCameraPoint(FVector2D PixelCoordinate)
 {
 	CameraImageCapture& CameraCapture = CameraImageCapture::Get();
 	
-	DirectX::XMFLOAT2 cameraPoint = DirectX::XMFLOAT2(pixelCoordinate.X, pixelCoordinate.Y);
-	DirectX::XMFLOAT2 unprojectedPointAtUnitDepth = CameraCapture.UnprojectPVCamPointAtUnitDepth(cameraPoint);
+	DirectX::XMFLOAT2 CameraPoint = DirectX::XMFLOAT2(PixelCoordinate.X, PixelCoordinate.Y);
+	DirectX::XMFLOAT2 UnprojectedPointAtUnitDepth = CameraCapture.UnprojectPVCamPointAtUnitDepth(CameraPoint);
 	
-	FVector ray = WindowsMixedReality::WMRUtility::FromMixedRealityVector(
+	FVector Ray = WindowsMixedReality::WMRUtility::FromMixedRealityVector(
 		DirectX::XMFLOAT3(
-			unprojectedPointAtUnitDepth.x,
-			unprojectedPointAtUnitDepth.y,
+			UnprojectedPointAtUnitDepth.x,
+			UnprojectedPointAtUnitDepth.y,
 			-1.0f // Unprojection happened at 1 meter
 		)
 	) * 100.0f;
 
-	ray.Normalize();
+	Ray.Normalize();
 
 	FScopeLock sl(&PVCamToWorldLock);
-	return PVCameraToWorldMatrix.TransformVector(ray);
+	return PVCameraToWorldMatrix.TransformVector(Ray);
 }
 
 void FHoloLensARSystem::StartCameraCapture()
diff --git a/Engine/Plugins/Runtime/WindowsMixedReality/Source/HoloLensAR/Public/HoloLensARSystem.h b/Engine/Plugins/Runtime/WindowsMixedReality/Source/HoloLensAR/Public/HoloLensARSystem.h
--- a/Engine/Plugins/Runtime/WindowsMixedReality/Source/HoloLensAR/Public/HoloLensARSystem.h
+++ b/Engine/Plugins/Runtime/WindowsMixedReality/Source/HoloLensAR/Public/HoloLensARSystem.h
@@ -117,6 +117,7 @@ private:
 	virtual TArray<FARVideoFormat> OnGetSupportedVideoFormats(EARSessionType SessionType) const override;
 	virtual TArray<FVector> OnGetPointCloud() const override;
 	virtual bool OnAddRuntimeCandidateImage(UARSessionConfig* SessionConfig, UTexture2D* CandidateTexture, FString FriendlyName, float PhysicalWidth) override;
+	virtual bool OnGetCameraIntrinsics(FARCameraIntrinsics& OutCameraIntrinsics) const;
 
 	// @todo JoeG - Figure out why we have these and if we really need them
 	virtual void* GetARSessionRawPointer() override { return nullptr; }
