Remove unneeded forward declares causing compile issues

diff --git a/Engine/Plugins/Runtime/WindowsMixedReality/Source/HoloLensAR/Private/HoloLensARSystem.cpp b/Engine/Plugins/Runtime/WindowsMixedReality/Source/HoloLensAR/Private/HoloLensARSystem.cpp
--- a/Engine/Plugins/Runtime/WindowsMixedReality/Source/HoloLensAR/Private/HoloLensARSystem.cpp
+++ b/Engine/Plugins/Runtime/WindowsMixedReality/Source/HoloLensAR/Private/HoloLensARSystem.cpp
@@ -1472,7 +1472,7 @@ bool FHoloLensARSystem::HasHandMeshData() const
 	return false;
 }
 
-bool FHoloLensARSystem::GetHandMeshData(EControllerHand Hand, TArray<struct FVector>& OutVertices, TArray<struct FVector>& OutNormals, TArray<int32>& OutIndices, FTransform& OutHandMeshTransform) const
+bool FHoloLensARSystem::GetHandMeshData(EControllerHand Hand, TArray<FVector>& OutVertices, TArray<FVector>& OutNormals, TArray<int32>& OutIndices, FTransform& OutHandMeshTransform) const
 {
 	FMeshUpdate HandState = HandMeshes[(int)Hand];
 
diff --git a/Engine/Plugins/Runtime/WindowsMixedReality/Source/HoloLensAR/Public/HoloLensARSystem.h b/Engine/Plugins/Runtime/WindowsMixedReality/Source/HoloLensAR/Public/HoloLensARSystem.h
--- a/Engine/Plugins/Runtime/WindowsMixedReality/Source/HoloLensAR/Public/HoloLensARSystem.h
+++ b/Engine/Plugins/Runtime/WindowsMixedReality/Source/HoloLensAR/Public/HoloLensARSystem.h
@@ -291,10 +291,10 @@ private:
 	{
 		return false;
 	}
-	virtual bool GetAllKeypointStates(EControllerHand Hand, TArray<struct FVector>& OutPositions, TArray<struct FQuat>& OutRotations, TArray<float>& OutRadii) const override
+	virtual bool GetAllKeypointStates(EControllerHand Hand, TArray<FVector>& OutPositions, TArray<FQuat>& OutRotations, TArray<float>& OutRadii) const override
 	{
 		return false;
 	}
 	virtual bool HasHandMeshData() const override;
-	virtual bool GetHandMeshData(EControllerHand Hand, TArray<struct FVector>& OutVertices, TArray<struct FVector>& OutNormals, TArray<int32>& OutIndices, FTransform& OutHandMeshTransform) const override;
+	virtual bool GetHandMeshData(EControllerHand Hand, TArray<FVector>& OutVertices, TArray<FVector>& OutNormals, TArray<int32>& OutIndices, FTransform& OutHandMeshTransform) const override;
 };
