PR #8651: Update WMR RHI to use a non-null adapter (Contributed by [REDACTED_PERSON])

diff --git a/Engine/Plugins/Runtime/WindowsMixedReality/Source/WindowsMixedRealityRHI/Private/WindowsMixedRealityRHI.cpp b/Engine/Plugins/Runtime/WindowsMixedReality/Source/WindowsMixedRealityRHI/Private/WindowsMixedRealityRHI.cpp
--- a/Engine/Plugins/Runtime/WindowsMixedReality/Source/WindowsMixedRealityRHI/Private/WindowsMixedRealityRHI.cpp
+++ b/Engine/Plugins/Runtime/WindowsMixedReality/Source/WindowsMixedRealityRHI/Private/WindowsMixedRealityRHI.cpp
@@ -50,7 +50,7 @@ void FWindowsMixedRealityRHIModule::StartupModule()
 
 #endif
 
-	ComPtr<IDXGIAdapter> TempAdapter;
+	TRefCountPtr<IDXGIAdapter> TempAdapter;
 
 	VERIFYD3D11RESULT(CreateDXGIFactory1(__uuidof(IDXGIFactory1), &DXGIFactory1));
 
@@ -62,7 +62,7 @@ void FWindowsMixedRealityRHIModule::StartupModule()
 	};
 
 
-	for (uint32 AdapterIndex = 0; DXGIFactory1->EnumAdapters(AdapterIndex, &TempAdapter) != DXGI_ERROR_NOT_FOUND; ++AdapterIndex)
+	for (uint32 AdapterIndex = 0; DXGIFactory1->EnumAdapters(AdapterIndex, TempAdapter.GetInitReference()) != DXGI_ERROR_NOT_FOUND; ++AdapterIndex)
 	{
 		if (!TempAdapter)
 		{
@@ -89,7 +89,7 @@ void FWindowsMixedRealityRHIModule::StartupModule()
 #endif //UE_BUILD_DEBUG
 
 		if (FAILED(D3D11CreateDevice(
-			TempAdapter.Get(),
+			TempAdapter.GetReference(),
 			D3D_DRIVER_TYPE_UNKNOWN,
 			NULL,
 			DeviceCreationFlags,
@@ -104,7 +104,8 @@ void FWindowsMixedRealityRHIModule::StartupModule()
 			continue; 
 		}
 
-		ChosenAdapter.MaxSupportedFeatureLevel = OutFeatureLevel;
+		FD3D11Adapter _adapter(TempAdapter, OutFeatureLevel, false, false);
+		ChosenAdapter = _adapter;
 		ChosenDescription = AdapterDesc;
 		break;
 	}
