Unshelved from pending changelist '18540094':

PR #8767: Deprecate WMR (Contributed by [REDACTED_PERSON])

diff --git a/Engine/Plugins/Runtime/AR/AzureSpatialAnchorsForWMR/AzureSpatialAnchorsForWMR.uplugin b/Engine/Plugins/Runtime/AR/AzureSpatialAnchorsForWMR/AzureSpatialAnchorsForWMR.uplugin
--- a/Engine/Plugins/Runtime/AR/AzureSpatialAnchorsForWMR/AzureSpatialAnchorsForWMR.uplugin
+++ b/Engine/Plugins/Runtime/AR/AzureSpatialAnchorsForWMR/AzureSpatialAnchorsForWMR.uplugin
@@ -3,7 +3,7 @@
 	"Version" : 1,
 	"VersionName" : "1.0",
 	"FriendlyName" : "Azure Spatial Anchors For WMR",
-	"Description" : "Microsoft Azure Spatial Anchors implementation for Windows Mixed Reality",
+	"Description" : "Microsoft Azure Spatial Anchors implementation for Windows Mixed Reality. This plugin is deprecated and will be removed in a future engine release. Please use the Microsoft OpenXR plugin from the Marketplace instead.",
 	"Category" : "Augmented Reality",
 	"CreatedBy" : "Epic Games",
 	"CreatedByURL" : "epicgames.com",
diff --git a/Engine/Plugins/Runtime/AR/AzureSpatialAnchorsForWMR/Source/AzureSpatialAnchorsForWMR/Private/AzureSpatialAnchorsForWMR.cpp b/Engine/Plugins/Runtime/AR/AzureSpatialAnchorsForWMR/Source/AzureSpatialAnchorsForWMR/Private/AzureSpatialAnchorsForWMR.cpp
--- a/Engine/Plugins/Runtime/AR/AzureSpatialAnchorsForWMR/Source/AzureSpatialAnchorsForWMR/Private/AzureSpatialAnchorsForWMR.cpp
+++ b/Engine/Plugins/Runtime/AR/AzureSpatialAnchorsForWMR/Source/AzureSpatialAnchorsForWMR/Private/AzureSpatialAnchorsForWMR.cpp
@@ -384,7 +384,7 @@ void FAzureSpatialAnchorsForWMR::SubmitDiagnosticsManifestAsync(const FString& M
 
 void FAzureSpatialAnchorsForWMR::CreateNamedARPinAroundAnchor(const FString& InLocalAnchorId, UARPin*& OutARPin)
 {
-	OutARPin = UHoloLensARFunctionLibrary::CreateNamedARPinAroundAnchor(FName(InLocalAnchorId), InLocalAnchorId);
+	OutARPin = UDEPRECATED_HoloLensARFunctionLibrary::CreateNamedARPinAroundAnchor(FName(InLocalAnchorId), InLocalAnchorId);
 }
 
 bool FAzureSpatialAnchorsForWMR::CreateARPinAroundAzureCloudSpatialAnchor(const FString& InPinId, UAzureCloudSpatialAnchor* InAzureCloudSpatialAnchor, UARPin*& OutARPin)
@@ -423,7 +423,7 @@ bool FAzureSpatialAnchorsForWMR::CreateARPinAroundAzureCloudSpatialAnchor(const
 	}
 
 	// create a pin around the local anchor
-	OutARPin = UHoloLensARFunctionLibrary::CreateNamedARPinAroundAnchor(PinIdName, InPinId);
+	OutARPin = UDEPRECATED_HoloLensARFunctionLibrary::CreateNamedARPinAroundAnchor(PinIdName, InPinId);
 	return (OutARPin != nullptr);
 	{
 		InAzureCloudSpatialAnchor->ARPin = OutARPin;
diff --git a/Engine/Plugins/Runtime/WindowsMixedReality/Source/HoloLensAR/Private/HoloLensARFunctionLibrary.cpp b/Engine/Plugins/Runtime/WindowsMixedReality/Source/HoloLensAR/Private/HoloLensARFunctionLibrary.cpp
--- a/Engine/Plugins/Runtime/WindowsMixedReality/Source/HoloLensAR/Private/HoloLensARFunctionLibrary.cpp
+++ b/Engine/Plugins/Runtime/WindowsMixedReality/Source/HoloLensAR/Private/HoloLensARFunctionLibrary.cpp
@@ -4,7 +4,7 @@
 #include "HoloLensModule.h"
 #include "Engine/Engine.h"
 
-UWMRARPin* UHoloLensARFunctionLibrary::CreateNamedARPin(FName Name, const FTransform& PinToWorldTransform)
+UWMRARPin* UDEPRECATED_HoloLensARFunctionLibrary::CreateNamedARPin(FName Name, const FTransform& PinToWorldTransform)
 {
 	TSharedPtr<FHoloLensARSystem, ESPMode::ThreadSafe> ARSystem = FHoloLensModuleAR::GetHoloLensARSystem();
 	if (!ARSystem.IsValid())
@@ -14,7 +14,7 @@ UWMRARPin* UHoloLensARFunctionLibrary::CreateNamedARPin(FName Name, const FTrans
 	return ARSystem->WMRCreateNamedARPin(Name, PinToWorldTransform);
 }
 
-TArray<UWMRARPin*> UHoloLensARFunctionLibrary::LoadWMRAnchorStoreARPins()
+TArray<UWMRARPin*> UDEPRECATED_HoloLensARFunctionLibrary::LoadWMRAnchorStoreARPins()
 {
 	TSharedPtr<FHoloLensARSystem, ESPMode::ThreadSafe> ARSystem = FHoloLensModuleAR::GetHoloLensARSystem();
 	if (!ARSystem.IsValid())
@@ -25,7 +25,7 @@ TArray<UWMRARPin*> UHoloLensARFunctionLibrary::LoadWMRAnchorStoreARPins()
 	return ARSystem->WMRLoadWMRAnchorStoreARPins();
 }
 
-bool UHoloLensARFunctionLibrary::SaveARPinToWMRAnchorStore(UARPin* InPin)
+bool UDEPRECATED_HoloLensARFunctionLibrary::SaveARPinToWMRAnchorStore(UARPin* InPin)
 {
 	TSharedPtr<FHoloLensARSystem, ESPMode::ThreadSafe> ARSystem = FHoloLensModuleAR::GetHoloLensARSystem();
 	if (!ARSystem.IsValid())
@@ -41,7 +41,7 @@ bool UHoloLensARFunctionLibrary::SaveARPinToWMRAnchorStore(UARPin* InPin)
 	return ARSystem->WMRSaveARPinToAnchorStore(InPin);
 }
 
-void UHoloLensARFunctionLibrary::RemoveARPinFromWMRAnchorStore(UARPin* InPin)
+void UDEPRECATED_HoloLensARFunctionLibrary::RemoveARPinFromWMRAnchorStore(UARPin* InPin)
 {
 	TSharedPtr<FHoloLensARSystem, ESPMode::ThreadSafe> ARSystem = FHoloLensModuleAR::GetHoloLensARSystem();
 	if (!ARSystem.IsValid())
@@ -58,7 +58,7 @@ void UHoloLensARFunctionLibrary::RemoveARPinFromWMRAnchorStore(UARPin* InPin)
 }
 
 
-void UHoloLensARFunctionLibrary::SetEnabledMixedRealityCamera(bool IsEnabled)
+void UDEPRECATED_HoloLensARFunctionLibrary::SetEnabledMixedRealityCamera(bool IsEnabled)
 {
 	TSharedPtr<FHoloLensARSystem, ESPMode::ThreadSafe> ARSystem = FHoloLensModuleAR::GetHoloLensARSystem();
 	if (!ARSystem.IsValid())
@@ -69,7 +69,7 @@ void UHoloLensARFunctionLibrary::SetEnabledMixedRealityCamera(bool IsEnabled)
 }
 
 
-FIntPoint UHoloLensARFunctionLibrary::ResizeMixedRealityCamera(const FIntPoint& size)
+FIntPoint UDEPRECATED_HoloLensARFunctionLibrary::ResizeMixedRealityCamera(const FIntPoint& size)
 {
 	TSharedPtr<FHoloLensARSystem, ESPMode::ThreadSafe> ARSystem = FHoloLensModuleAR::GetHoloLensARSystem();
 	if (!ARSystem.IsValid())
@@ -81,7 +81,7 @@ FIntPoint UHoloLensARFunctionLibrary::ResizeMixedRealityCamera(const FIntPoint&
 	return newSize;
 }
 
-FTransform UHoloLensARFunctionLibrary::GetPVCameraToWorldTransform()
+FTransform UDEPRECATED_HoloLensARFunctionLibrary::GetPVCameraToWorldTransform()
 {
 	TSharedPtr<FHoloLensARSystem, ESPMode::ThreadSafe> ARSystem = FHoloLensModuleAR::GetHoloLensARSystem();
 	if (!ARSystem.IsValid())
@@ -92,7 +92,7 @@ FTransform UHoloLensARFunctionLibrary::GetPVCameraToWorldTransform()
 	return ARSystem->GetPVCameraToWorldTransform();
 }
 
-bool UHoloLensARFunctionLibrary::GetPVCameraIntrinsics(FVector2D& focalLength, int& width, int& height, FVector2D& principalPoint, FVector& radialDistortion, FVector2D& tangentialDistortion)
+bool UDEPRECATED_HoloLensARFunctionLibrary::GetPVCameraIntrinsics(FVector2D& focalLength, int& width, int& height, FVector2D& principalPoint, FVector& radialDistortion, FVector2D& tangentialDistortion)
 {
 	TSharedPtr<FHoloLensARSystem, ESPMode::ThreadSafe> ARSystem = FHoloLensModuleAR::GetHoloLensARSystem();
 	if (!ARSystem.IsValid())
@@ -103,7 +103,7 @@ bool UHoloLensARFunctionLibrary::GetPVCameraIntrinsics(FVector2D& focalLength, i
 	return ARSystem->GetPVCameraIntrinsics(focalLength, width, height, principalPoint, radialDistortion, tangentialDistortion);
 }
 
-FVector UHoloLensARFunctionLibrary::GetWorldSpaceRayFromCameraPoint(FVector2D pixelCoordinate)
+FVector UDEPRECATED_HoloLensARFunctionLibrary::GetWorldSpaceRayFromCameraPoint(FVector2D pixelCoordinate)
 {
 	TSharedPtr<FHoloLensARSystem, ESPMode::ThreadSafe> ARSystem = FHoloLensModuleAR::GetHoloLensARSystem();
 	if (!ARSystem.IsValid())
@@ -114,7 +114,7 @@ FVector UHoloLensARFunctionLibrary::GetWorldSpaceRayFromCameraPoint(FVector2D pi
 	return ARSystem->GetWorldSpaceRayFromCameraPoint(pixelCoordinate);
 }
 
-void UHoloLensARFunctionLibrary::StartCameraCapture()
+void UDEPRECATED_HoloLensARFunctionLibrary::StartCameraCapture()
 {
 	TSharedPtr<FHoloLensARSystem, ESPMode::ThreadSafe> ARSystem = FHoloLensModuleAR::GetHoloLensARSystem();
 	if (!ARSystem.IsValid())
@@ -125,7 +125,7 @@ void UHoloLensARFunctionLibrary::StartCameraCapture()
 	ARSystem->StartCameraCapture();
 }
 
-void UHoloLensARFunctionLibrary::StopCameraCapture()
+void UDEPRECATED_HoloLensARFunctionLibrary::StopCameraCapture()
 {
 	TSharedPtr<FHoloLensARSystem, ESPMode::ThreadSafe> ARSystem = FHoloLensModuleAR::GetHoloLensARSystem();
 	if (!ARSystem.IsValid())
@@ -136,7 +136,7 @@ void UHoloLensARFunctionLibrary::StopCameraCapture()
 	ARSystem->StopCameraCapture();
 }
 
-void UHoloLensARFunctionLibrary::StartQRCodeCapture()
+void UDEPRECATED_HoloLensARFunctionLibrary::StartQRCodeCapture()
 {
 	TSharedPtr<FHoloLensARSystem, ESPMode::ThreadSafe> ARSystem = FHoloLensModuleAR::GetHoloLensARSystem();
 	if (!ARSystem.IsValid())
@@ -147,7 +147,7 @@ void UHoloLensARFunctionLibrary::StartQRCodeCapture()
 	ARSystem->SetupQRCodeTracking();
 }
 
-void UHoloLensARFunctionLibrary::StopQRCodeCapture()
+void UDEPRECATED_HoloLensARFunctionLibrary::StopQRCodeCapture()
 {
 	TSharedPtr<FHoloLensARSystem, ESPMode::ThreadSafe> ARSystem = FHoloLensModuleAR::GetHoloLensARSystem();
 	if (!ARSystem.IsValid())
@@ -158,17 +158,17 @@ void UHoloLensARFunctionLibrary::StopQRCodeCapture()
 	ARSystem->StopQRCodeTracking();
 }
 
-bool UHoloLensARFunctionLibrary::ShowKeyboard()
+bool UDEPRECATED_HoloLensARFunctionLibrary::ShowKeyboard()
 {
 	return false;
 }
 
-bool UHoloLensARFunctionLibrary::HideKeyboard()
+bool UDEPRECATED_HoloLensARFunctionLibrary::HideKeyboard()
 {
 	return false;
 }
 
-UWMRARPin* UHoloLensARFunctionLibrary::CreateNamedARPinAroundAnchor(FName Name, const FString& AnchorId)
+UWMRARPin* UDEPRECATED_HoloLensARFunctionLibrary::CreateNamedARPinAroundAnchor(FName Name, const FString& AnchorId)
 {
 	TSharedPtr<FHoloLensARSystem, ESPMode::ThreadSafe> ARSystem = FHoloLensModuleAR::GetHoloLensARSystem();
 	if (!ARSystem.IsValid())
@@ -178,7 +178,7 @@ UWMRARPin* UHoloLensARFunctionLibrary::CreateNamedARPinAroundAnchor(FName Name,
 	return ARSystem->WMRCreateNamedARPinAroundAnchor(Name, AnchorId);
 }
 
-void UHoloLensARFunctionLibrary::SetUseLegacyHandMeshVisualization(bool bUseLegacyHandMeshVisualization)
+void UDEPRECATED_HoloLensARFunctionLibrary::SetUseLegacyHandMeshVisualization(bool bUseLegacyHandMeshVisualization)
 {
 	TSharedPtr<FHoloLensARSystem, ESPMode::ThreadSafe> ARSystem = FHoloLensModuleAR::GetHoloLensARSystem();
 	if (!ARSystem.IsValid())
diff --git a/Engine/Plugins/Runtime/WindowsMixedReality/Source/HoloLensAR/Private/HoloLensARSystem.cpp b/Engine/Plugins/Runtime/WindowsMixedReality/Source/HoloLensAR/Private/HoloLensARSystem.cpp
--- a/Engine/Plugins/Runtime/WindowsMixedReality/Source/HoloLensAR/Private/HoloLensARSystem.cpp
+++ b/Engine/Plugins/Runtime/WindowsMixedReality/Source/HoloLensAR/Private/HoloLensARSystem.cpp
@@ -479,21 +479,21 @@ bool FHoloLensARSystem::OnToggleARCapture(const bool bOnOff, const EARCaptureTyp
 		case EARCaptureType::Camera:
 			if (bOnOff)
 			{
-				UHoloLensARFunctionLibrary::StartCameraCapture();
+				UDEPRECATED_HoloLensARFunctionLibrary::StartCameraCapture();
 			}
 			else
 			{
-				UHoloLensARFunctionLibrary::StopCameraCapture();
+				UDEPRECATED_HoloLensARFunctionLibrary::StopCameraCapture();
 			}
 			break;
 		case EARCaptureType::QRCode:
 			if (bOnOff)
 			{
-				UHoloLensARFunctionLibrary::StartQRCodeCapture();
+				UDEPRECATED_HoloLensARFunctionLibrary::StartQRCodeCapture();
 			}
 			else
 			{
-				UHoloLensARFunctionLibrary::StopQRCodeCapture();
+				UDEPRECATED_HoloLensARFunctionLibrary::StopQRCodeCapture();
 			}
 			break;
 		default:
@@ -540,12 +540,12 @@ bool FHoloLensARSystem::OnToggleARCapture(const bool bOnOff, const EARCaptureTyp
 
 void FHoloLensARSystem::OnSetEnabledXRCamera(bool bOnOff)
 {
-	UHoloLensARFunctionLibrary::SetEnabledMixedRealityCamera(bOnOff);
+	UDEPRECATED_HoloLensARFunctionLibrary::SetEnabledMixedRealityCamera(bOnOff);
 }
 
 FIntPoint FHoloLensARSystem::OnResizeXRCamera(const FIntPoint& InSize)
 {
-	return UHoloLensARFunctionLibrary::ResizeMixedRealityCamera(InSize);
+	return UDEPRECATED_HoloLensARFunctionLibrary::ResizeMixedRealityCamera(InSize);
 }
 
 
diff --git a/Engine/Plugins/Runtime/WindowsMixedReality/Source/HoloLensAR/Public/HoloLensARFunctionLibrary.h b/Engine/Plugins/Runtime/WindowsMixedReality/Source/HoloLensAR/Public/HoloLensARFunctionLibrary.h
--- a/Engine/Plugins/Runtime/WindowsMixedReality/Source/HoloLensAR/Public/HoloLensARFunctionLibrary.h
+++ b/Engine/Plugins/Runtime/WindowsMixedReality/Source/HoloLensAR/Public/HoloLensARFunctionLibrary.h
@@ -10,8 +10,8 @@
 #include "HoloLensARFunctionLibrary.generated.h"
 
 /** A function library that provides static/Blueprint functions for HoloLensAR.*/
-UCLASS()
-class HOLOLENSAR_API UHoloLensARFunctionLibrary : public UBlueprintFunctionLibrary
+UCLASS(Deprecated, meta = (DeprecationMessage = "WMR has been deprecated, please use the Microsoft OpenXR plugin from the Marketplace."))
+class HOLOLENSAR_API UDEPRECATED_HoloLensARFunctionLibrary : public UBlueprintFunctionLibrary
 {
 	GENERATED_BODY()
 
diff --git a/Engine/Plugins/Runtime/WindowsMixedReality/Source/WindowsMixedRealityHMD/Private/WindowsMixedRealityFunctionLibrary.cpp b/Engine/Plugins/Runtime/WindowsMixedReality/Source/WindowsMixedRealityHMD/Private/WindowsMixedRealityFunctionLibrary.cpp
--- a/Engine/Plugins/Runtime/WindowsMixedReality/Source/WindowsMixedRealityHMD/Private/WindowsMixedRealityFunctionLibrary.cpp
+++ b/Engine/Plugins/Runtime/WindowsMixedReality/Source/WindowsMixedRealityHMD/Private/WindowsMixedRealityFunctionLibrary.cpp
@@ -12,7 +12,7 @@
 
 #include <functional>
 
-UWindowsMixedRealityFunctionLibrary::UWindowsMixedRealityFunctionLibrary(const FObjectInitializer& ObjectInitializer)
+UDEPRECATED_WindowsMixedRealityFunctionLibrary::UDEPRECATED_WindowsMixedRealityFunctionLibrary(const FObjectInitializer& ObjectInitializer)
 	: Super(ObjectInitializer)
 {
 }
@@ -27,7 +27,7 @@ WindowsMixedReality::FWindowsMixedRealityHMD* GetWindowsMixedRealityHMD() noexce
 	return nullptr;
 }
 
-FString UWindowsMixedRealityFunctionLibrary::GetVersionString()
+FString UDEPRECATED_WindowsMixedRealityFunctionLibrary::GetVersionString()
 {
 	WindowsMixedReality::FWindowsMixedRealityHMD* hmd = GetWindowsMixedRealityHMD();
 	if (hmd == nullptr)
@@ -38,7 +38,7 @@ FString UWindowsMixedRealityFunctionLibrary::GetVersionString()
 	return hmd->GetVersionString();
 }
 
-void UWindowsMixedRealityFunctionLibrary::ToggleImmersive(bool immersive)
+void UDEPRECATED_WindowsMixedRealityFunctionLibrary::ToggleImmersive(bool immersive)
 {
 	WindowsMixedReality::FWindowsMixedRealityHMD* hmd = GetWindowsMixedRealityHMD();
 	if (hmd == nullptr)
@@ -49,7 +49,7 @@ void UWindowsMixedRealityFunctionLibrary::ToggleImmersive(bool immersive)
 	hmd->EnableStereo(immersive);
 }
 
-bool UWindowsMixedRealityFunctionLibrary::IsCurrentlyImmersive()
+bool UDEPRECATED_WindowsMixedRealityFunctionLibrary::IsCurrentlyImmersive()
 {
 	WindowsMixedReality::FWindowsMixedRealityHMD* hmd = GetWindowsMixedRealityHMD();
 	if (hmd == nullptr)
@@ -60,7 +60,7 @@ bool UWindowsMixedRealityFunctionLibrary::IsCurrentlyImmersive()
 	return hmd->IsCurrentlyImmersive();
 }
 
-bool UWindowsMixedRealityFunctionLibrary::IsDisplayOpaque()
+bool UDEPRECATED_WindowsMixedRealityFunctionLibrary::IsDisplayOpaque()
 {
 	WindowsMixedReality::FWindowsMixedRealityHMD* hmd = GetWindowsMixedRealityHMD();
 	if (hmd == nullptr)
@@ -71,7 +71,7 @@ bool UWindowsMixedRealityFunctionLibrary::IsDisplayOpaque()
 	return hmd->IsDisplayOpaque();
 }
 
-void UWindowsMixedRealityFunctionLibrary::LockMouseToCenter(bool locked)
+void UDEPRECATED_WindowsMixedRealityFunctionLibrary::LockMouseToCenter(bool locked)
 {
 	WindowsMixedReality::FWindowsMixedRealityHMD* hmd = GetWindowsMixedRealityHMD();
 	if (hmd == nullptr)
@@ -82,7 +82,7 @@ void UWindowsMixedRealityFunctionLibrary::LockMouseToCenter(bool locked)
 	hmd->LockMouseToCenter(locked);
 }
 
-bool UWindowsMixedRealityFunctionLibrary::IsTrackingAvailable()
+bool UDEPRECATED_WindowsMixedRealityFunctionLibrary::IsTrackingAvailable()
 {
 #if WITH_WINDOWS_MIXED_REALITY
 	WindowsMixedReality::FWindowsMixedRealityHMD* hmd = GetWindowsMixedRealityHMD();
@@ -97,7 +97,7 @@ bool UWindowsMixedRealityFunctionLibrary::IsTrackingAvailable()
 #endif
 }
 
-FPointerPoseInfo UWindowsMixedRealityFunctionLibrary::GetPointerPoseInfo(EControllerHand hand)
+FPointerPoseInfo UDEPRECATED_WindowsMixedRealityFunctionLibrary::GetPointerPoseInfo(EControllerHand hand)
 {
 	FPointerPoseInfo info;
 	info.TrackingStatus = GetControllerTrackingStatus(hand);
@@ -122,7 +122,7 @@ FPointerPoseInfo UWindowsMixedRealityFunctionLibrary::GetPointerPoseInfo(EContro
 	return info;
 }
 
-bool UWindowsMixedRealityFunctionLibrary::IsButtonClicked(EControllerHand hand, EHMDInputControllerButtons button)
+bool UDEPRECATED_WindowsMixedRealityFunctionLibrary::IsButtonClicked(EControllerHand hand, EHMDInputControllerButtons button)
 {
 #if WITH_WINDOWS_MIXED_REALITY
 	WindowsMixedReality::FWindowsMixedRealityHMD* hmd = GetWindowsMixedRealityHMD();
@@ -137,7 +137,7 @@ bool UWindowsMixedRealityFunctionLibrary::IsButtonClicked(EControllerHand hand,
 	return false;
 }
 
-bool UWindowsMixedRealityFunctionLibrary::IsButtonDown(EControllerHand hand, EHMDInputControllerButtons button)
+bool UDEPRECATED_WindowsMixedRealityFunctionLibrary::IsButtonDown(EControllerHand hand, EHMDInputControllerButtons button)
 {
 #if WITH_WINDOWS_MIXED_REALITY
 	WindowsMixedReality::FWindowsMixedRealityHMD* hmd = GetWindowsMixedRealityHMD();
@@ -152,17 +152,17 @@ bool UWindowsMixedRealityFunctionLibrary::IsButtonDown(EControllerHand hand, EHM
 	return false;
 }
 
-bool UWindowsMixedRealityFunctionLibrary::IsGrasped(EControllerHand hand)
+bool UDEPRECATED_WindowsMixedRealityFunctionLibrary::IsGrasped(EControllerHand hand)
 {
 	return IsButtonDown(hand, EHMDInputControllerButtons::Grasp);
 }
 
-bool UWindowsMixedRealityFunctionLibrary::IsSelectPressed(EControllerHand hand)
+bool UDEPRECATED_WindowsMixedRealityFunctionLibrary::IsSelectPressed(EControllerHand hand)
 {
 	return IsButtonDown(hand, EHMDInputControllerButtons::Select);
 }
 
-EHMDTrackingStatus UWindowsMixedRealityFunctionLibrary::GetControllerTrackingStatus(EControllerHand hand)
+EHMDTrackingStatus UDEPRECATED_WindowsMixedRealityFunctionLibrary::GetControllerTrackingStatus(EControllerHand hand)
 {
 #if WITH_WINDOWS_MIXED_REALITY
 	WindowsMixedReality::FWindowsMixedRealityHMD* hmd = GetWindowsMixedRealityHMD();
@@ -177,7 +177,7 @@ EHMDTrackingStatus UWindowsMixedRealityFunctionLibrary::GetControllerTrackingSta
 	return EHMDTrackingStatus::NotTracked;
 }
 
-void UWindowsMixedRealityFunctionLibrary::SetFocusPointForFrame(FVector position)
+void UDEPRECATED_WindowsMixedRealityFunctionLibrary::SetFocusPointForFrame(FVector position)
 {
 #if WITH_WINDOWS_MIXED_REALITY
 	WindowsMixedReality::FWindowsMixedRealityHMD* hmd = GetWindowsMixedRealityHMD();
diff --git a/Engine/Plugins/Runtime/WindowsMixedReality/Source/WindowsMixedRealityHMD/Private/WindowsMixedRealityHMD.cpp b/Engine/Plugins/Runtime/WindowsMixedReality/Source/WindowsMixedRealityHMD/Private/WindowsMixedRealityHMD.cpp
--- a/Engine/Plugins/Runtime/WindowsMixedReality/Source/WindowsMixedRealityHMD/Private/WindowsMixedRealityHMD.cpp
+++ b/Engine/Plugins/Runtime/WindowsMixedReality/Source/WindowsMixedRealityHMD/Private/WindowsMixedRealityHMD.cpp
@@ -226,7 +226,7 @@ namespace WindowsMixedReality
 		}
 
 #if WITH_INPUT_SIMULATION
-		if (auto* InputSim = UWindowsMixedRealityInputSimulationEngineSubsystem::GetInputSimulationIfEnabled())
+		if (auto* InputSim = UDEPRECATED_WindowsMixedRealityInputSimulationEngineSubsystem::GetInputSimulationIfEnabled())
 		{
 			if (InputSim->HasPositionalTracking())
 			{
@@ -378,7 +378,7 @@ namespace WindowsMixedReality
 	bool FWindowsMixedRealityHMD::HasValidTrackingPosition()
 	{
 #if WITH_INPUT_SIMULATION
-		if (auto* InputSim = UWindowsMixedRealityInputSimulationEngineSubsystem::GetInputSimulationIfEnabled())
+		if (auto* InputSim = UDEPRECATED_WindowsMixedRealityInputSimulationEngineSubsystem::GetInputSimulationIfEnabled())
 		{
 			return InputSim->HasPositionalTracking();
 		}
@@ -840,7 +840,7 @@ namespace WindowsMixedReality
 		}
 
 #if WITH_INPUT_SIMULATION
-		if (auto* InputSim = UWindowsMixedRealityInputSimulationEngineSubsystem::GetInputSimulationIfEnabled())
+		if (auto* InputSim = UDEPRECATED_WindowsMixedRealityInputSimulationEngineSubsystem::GetInputSimulationIfEnabled())
 		{
 			CurrentOrientation = InputSim->GetHeadOrientation();
 			CurrentPosition = InputSim->GetHeadPosition();
@@ -911,11 +911,11 @@ namespace WindowsMixedReality
 		}
 
 		//far pointing from elbow transform
-		FPointerPoseInfo PointerPoseInfo = UWindowsMixedRealityFunctionLibrary::GetPointerPoseInfo(Hand);
+		FPointerPoseInfo PointerPoseInfo = UDEPRECATED_WindowsMixedRealityFunctionLibrary::GetPointerPoseInfo(Hand);
 		MotionControllerData.AimPosition = PointerPoseInfo.Origin;
 		MotionControllerData.AimRotation = PointerPoseInfo.Orientation;
 
-		MotionControllerData.bIsGrasped = UWindowsMixedRealityFunctionLibrary::IsGrasped(Hand);
+		MotionControllerData.bIsGrasped = UDEPRECATED_WindowsMixedRealityFunctionLibrary::IsGrasped(Hand);
 
 		if (!SupportsHandTracking())
 		{
@@ -2320,7 +2320,7 @@ namespace WindowsMixedReality
 	HMDTrackingStatus FWindowsMixedRealityHMD::GetControllerTrackingStatus(HMDHand hand)
 	{
 #if WITH_INPUT_SIMULATION
-		if (auto* InputSim = UWindowsMixedRealityInputSimulationEngineSubsystem::GetInputSimulationIfEnabled())
+		if (auto* InputSim = UDEPRECATED_WindowsMixedRealityInputSimulationEngineSubsystem::GetInputSimulationIfEnabled())
 		{
 			return (HMDTrackingStatus)InputSim->GetControllerTrackingStatus((EControllerHand)hand);
 		}
@@ -2412,7 +2412,7 @@ namespace WindowsMixedReality
 	HMDInputPressState WindowsMixedReality::FWindowsMixedRealityHMD::GetPressState(HMDHand hand, HMDInputControllerButtons button, bool onlyRegisterClicks)
 	{
 #if WITH_INPUT_SIMULATION
-		if (auto* InputSim = UWindowsMixedRealityInputSimulationEngineSubsystem::GetInputSimulationIfEnabled())
+		if (auto* InputSim = UDEPRECATED_WindowsMixedRealityInputSimulationEngineSubsystem::GetInputSimulationIfEnabled())
 		{
 			bool IsPressed;
 			if (InputSim->GetPressState((EControllerHand)hand, (EHMDInputControllerButtons)button, onlyRegisterClicks, IsPressed))
@@ -2461,7 +2461,7 @@ namespace WindowsMixedReality
 	void FWindowsMixedRealityHMD::GetPointerPose(EControllerHand hand, PointerPoseInfo& pi)
 	{
 #if WITH_INPUT_SIMULATION
-		if (auto* InputSim = UWindowsMixedRealityInputSimulationEngineSubsystem::GetInputSimulationIfEnabled())
+		if (auto* InputSim = UDEPRECATED_WindowsMixedRealityInputSimulationEngineSubsystem::GetInputSimulationIfEnabled())
 		{
 			FWindowsMixedRealityInputSimulationPointerPose InputSimPointerPose;
 			if (InputSim->GetHandPointerPose(hand, InputSimPointerPose))
diff --git a/Engine/Plugins/Runtime/WindowsMixedReality/Source/WindowsMixedRealityHMD/Private/WindowsMixedRealityHMD.h b/Engine/Plugins/Runtime/WindowsMixedReality/Source/WindowsMixedRealityHMD/Private/WindowsMixedRealityHMD.h
--- a/Engine/Plugins/Runtime/WindowsMixedReality/Source/WindowsMixedRealityHMD/Private/WindowsMixedRealityHMD.h
+++ b/Engine/Plugins/Runtime/WindowsMixedReality/Source/WindowsMixedRealityHMD/Private/WindowsMixedRealityHMD.h
@@ -52,7 +52,7 @@ namespace WindowsMixedReality
 		virtual int32 GetXRSystemFlags() const override
 		{
 			int32 XRFlags = EXRSystemFlags::IsHeadMounted;
-			if (!UWindowsMixedRealityFunctionLibrary::IsDisplayOpaque())
+			if (!UDEPRECATED_WindowsMixedRealityFunctionLibrary::IsDisplayOpaque())
 			{
 				XRFlags |= EXRSystemFlags::IsAR;
 			}
diff --git a/Engine/Plugins/Runtime/WindowsMixedReality/Source/WindowsMixedRealityHMD/Public/WindowsMixedRealityFunctionLibrary.h b/Engine/Plugins/Runtime/WindowsMixedReality/Source/WindowsMixedRealityHMD/Public/WindowsMixedRealityFunctionLibrary.h
--- a/Engine/Plugins/Runtime/WindowsMixedReality/Source/WindowsMixedRealityHMD/Public/WindowsMixedRealityFunctionLibrary.h
+++ b/Engine/Plugins/Runtime/WindowsMixedReality/Source/WindowsMixedRealityHMD/Public/WindowsMixedRealityFunctionLibrary.h
@@ -64,8 +64,8 @@ DECLARE_DYNAMIC_DELEGATE_OneParam(FTrackingChangeCallback, EHMDSpatialLocatabili
 /**
 * Windows Mixed Reality Extensions Function Library
 */
-UCLASS()
-class WINDOWSMIXEDREALITYHMD_API UWindowsMixedRealityFunctionLibrary : public UBlueprintFunctionLibrary
+UCLASS(Deprecated, meta = (DeprecationMessage = "WMR has been deprecated, please use the Microsoft OpenXR plugin from the Marketplace."))
+class WINDOWSMIXEDREALITYHMD_API UDEPRECATED_WindowsMixedRealityFunctionLibrary : public UBlueprintFunctionLibrary
 {
 	GENERATED_UCLASS_BODY()
 
diff --git a/Engine/Plugins/Runtime/WindowsMixedReality/Source/WindowsMixedRealityHandTracking/Private/WindowsMixedRealityHandTracking.cpp b/Engine/Plugins/Runtime/WindowsMixedReality/Source/WindowsMixedRealityHandTracking/Private/WindowsMixedRealityHandTracking.cpp
--- a/Engine/Plugins/Runtime/WindowsMixedReality/Source/WindowsMixedRealityHandTracking/Private/WindowsMixedRealityHandTracking.cpp
+++ b/Engine/Plugins/Runtime/WindowsMixedReality/Source/WindowsMixedRealityHandTracking/Private/WindowsMixedRealityHandTracking.cpp
@@ -154,7 +154,7 @@ public:
 	virtual bool GetHandJointTransform(EControllerHand Hand, EWMRHandKeypoint Keypoint, FTransform& Transform, float& OutRadius) override
 	{
 		OutRadius = 0.0f;
-		return UWindowsMixedRealityHandTrackingFunctionLibrary::GetHandJointTransform(Hand, Keypoint, Transform, OutRadius);
+		return UDEPRECATED_WindowsMixedRealityHandTrackingFunctionLibrary::GetHandJointTransform(Hand, Keypoint, Transform, OutRadius);
 	}
 
 private:
@@ -390,7 +390,7 @@ bool FWindowsMixedRealityHandTracking::IsHandTrackingStateValid() const
 
 void FWindowsMixedRealityHandTracking::OnGetXRSystemFlags(int32& XRFlags)
 {
-	if (UWindowsMixedRealityHandTrackingFunctionLibrary::SupportsHandTracking())
+	if (UDEPRECATED_WindowsMixedRealityHandTrackingFunctionLibrary::SupportsHandTracking())
 	{
 		XRFlags |= EXRSystemFlags::SupportsHandTracking;
 	}
@@ -402,7 +402,7 @@ bool FWindowsMixedRealityHandTracking::GetKeypointTransform(EControllerHand Hand
 	bool gotTransform = false;
  
 #if WITH_INPUT_SIMULATION
-	if (auto* InputSim = UWindowsMixedRealityInputSimulationEngineSubsystem::GetInputSimulationIfEnabled())
+	if (auto* InputSim = UDEPRECATED_WindowsMixedRealityInputSimulationEngineSubsystem::GetInputSimulationIfEnabled())
 	{
 		gotTransform = InputSim->GetHandJointTransform(Hand, Keypoint, OutTransform);
 	}
@@ -425,7 +425,7 @@ bool FWindowsMixedRealityHandTracking::GetKeypointTransform(EControllerHand Hand
 bool FWindowsMixedRealityHandTracking::GetKeypointRadius(EControllerHand Hand, EWMRHandKeypoint Keypoint, float& OutRadius) const
 {
 #if WITH_INPUT_SIMULATION
-	if (auto* InputSim = UWindowsMixedRealityInputSimulationEngineSubsystem::GetInputSimulationIfEnabled())
+	if (auto* InputSim = UDEPRECATED_WindowsMixedRealityInputSimulationEngineSubsystem::GetInputSimulationIfEnabled())
 	{
 		return InputSim->GetHandJointRadius(Hand, Keypoint, OutRadius);
 	}
diff --git a/Engine/Plugins/Runtime/WindowsMixedReality/Source/WindowsMixedRealityHandTracking/Private/WindowsMixedRealityHandTrackingFunctionLibrary.cpp b/Engine/Plugins/Runtime/WindowsMixedReality/Source/WindowsMixedRealityHandTracking/Private/WindowsMixedRealityHandTrackingFunctionLibrary.cpp
--- a/Engine/Plugins/Runtime/WindowsMixedReality/Source/WindowsMixedRealityHandTracking/Private/WindowsMixedRealityHandTrackingFunctionLibrary.cpp
+++ b/Engine/Plugins/Runtime/WindowsMixedReality/Source/WindowsMixedRealityHandTracking/Private/WindowsMixedRealityHandTrackingFunctionLibrary.cpp
@@ -4,7 +4,7 @@
 #include "IWindowsMixedRealityHandTrackingPlugin.h"
 #include "WindowsMixedRealityHandTracking.h"
 
-bool UWindowsMixedRealityHandTrackingFunctionLibrary::GetHandJointTransform(EControllerHand Hand, EWMRHandKeypoint Keypoint, FTransform& OutTransform, float& OutRadius)
+bool UDEPRECATED_WindowsMixedRealityHandTrackingFunctionLibrary::GetHandJointTransform(EControllerHand Hand, EWMRHandKeypoint Keypoint, FTransform& OutTransform, float& OutRadius)
 {
 	TSharedPtr<FWindowsMixedRealityHandTracking> HandTracking = StaticCastSharedPtr<FWindowsMixedRealityHandTracking>(IWindowsMixedRealityHandTrackingModule::Get().GetInputDevice());
 
@@ -26,7 +26,7 @@ bool UWindowsMixedRealityHandTrackingFunctionLibrary::GetHandJointTransform(ECon
 	return false;
 }
 
-bool UWindowsMixedRealityHandTrackingFunctionLibrary::SupportsHandTracking()
+bool UDEPRECATED_WindowsMixedRealityHandTrackingFunctionLibrary::SupportsHandTracking()
 {
 	return WindowsMixedReality::FWindowsMixedRealityStatics::SupportsHandTracking();
 }
diff --git a/Engine/Plugins/Runtime/WindowsMixedReality/Source/WindowsMixedRealityHandTracking/Public/WindowsMixedRealityHandTrackingFunctionLibrary.h b/Engine/Plugins/Runtime/WindowsMixedReality/Source/WindowsMixedRealityHandTracking/Public/WindowsMixedRealityHandTrackingFunctionLibrary.h
--- a/Engine/Plugins/Runtime/WindowsMixedReality/Source/WindowsMixedRealityHandTracking/Public/WindowsMixedRealityHandTrackingFunctionLibrary.h
+++ b/Engine/Plugins/Runtime/WindowsMixedReality/Source/WindowsMixedRealityHandTracking/Public/WindowsMixedRealityHandTrackingFunctionLibrary.h
@@ -9,8 +9,8 @@
 #include "InputCoreTypes.h"
 #include "WindowsMixedRealityHandTrackingFunctionLibrary.generated.h"
 
-UCLASS(ClassGroup = WindowsMixedReality)
-class WINDOWSMIXEDREALITYHANDTRACKING_API UWindowsMixedRealityHandTrackingFunctionLibrary :
+UCLASS(ClassGroup = WindowsMixedReality, Deprecated, meta = (DeprecationMessage = "WMR has been deprecated, please use the Microsoft OpenXR plugin from the Marketplace."))
+class WINDOWSMIXEDREALITYHANDTRACKING_API UDEPRECATED_WindowsMixedRealityHandTrackingFunctionLibrary :
 	public UBlueprintFunctionLibrary
 {
 	GENERATED_BODY()
diff --git a/Engine/Plugins/Runtime/WindowsMixedReality/Source/WindowsMixedRealityInputSimulation/Private/WindowsMixedRealityInputSimulationEngineSubsystem.cpp b/Engine/Plugins/Runtime/WindowsMixedReality/Source/WindowsMixedRealityInputSimulation/Private/WindowsMixedRealityInputSimulationEngineSubsystem.cpp
--- a/Engine/Plugins/Runtime/WindowsMixedReality/Source/WindowsMixedRealityInputSimulation/Private/WindowsMixedRealityInputSimulationEngineSubsystem.cpp
+++ b/Engine/Plugins/Runtime/WindowsMixedReality/Source/WindowsMixedRealityInputSimulation/Private/WindowsMixedRealityInputSimulationEngineSubsystem.cpp
@@ -16,17 +16,17 @@
 
 #define LOCTEXT_NAMESPACE "WindowsMixedRealityInputSimulation"
 
-UWindowsMixedRealityInputSimulationEngineSubsystem* UWindowsMixedRealityInputSimulationEngineSubsystem::GetInputSimulationIfEnabled()
+UDEPRECATED_WindowsMixedRealityInputSimulationEngineSubsystem* UDEPRECATED_WindowsMixedRealityInputSimulationEngineSubsystem::GetInputSimulationIfEnabled()
 {
-	if (!UWindowsMixedRealityInputSimulationEngineSubsystem::IsInputSimulationEnabled())
+	if (!UDEPRECATED_WindowsMixedRealityInputSimulationEngineSubsystem::IsInputSimulationEnabled())
 	{
 		return nullptr;
 	}
 
-	return GEngine->GetEngineSubsystem<UWindowsMixedRealityInputSimulationEngineSubsystem>();
+	return GEngine->GetEngineSubsystem<UDEPRECATED_WindowsMixedRealityInputSimulationEngineSubsystem>();
 }
 
-bool UWindowsMixedRealityInputSimulationEngineSubsystem::IsInputSimulationEnabled()
+bool UDEPRECATED_WindowsMixedRealityInputSimulationEngineSubsystem::IsInputSimulationEnabled()
 {
 #if WITH_EDITOR
 	// Only use input simulation when enabled
@@ -63,25 +63,25 @@ bool UWindowsMixedRealityInputSimulationEngineSubsystem::IsInputSimulationEnable
 	return false;
 }
 
-bool UWindowsMixedRealityInputSimulationEngineSubsystem::HasPositionalTracking() const
+bool UDEPRECATED_WindowsMixedRealityInputSimulationEngineSubsystem::HasPositionalTracking() const
 {
 	FRWScopeLock ReadLock(DataMutex, SLT_ReadOnly);
 	return bHasPositionalTracking;
 }
 
-const FQuat& UWindowsMixedRealityInputSimulationEngineSubsystem::GetHeadOrientation() const
+const FQuat& UDEPRECATED_WindowsMixedRealityInputSimulationEngineSubsystem::GetHeadOrientation() const
 {
 	FRWScopeLock ReadLock(DataMutex, SLT_ReadOnly);
 	return HeadOrientation;
 }
 
-const FVector& UWindowsMixedRealityInputSimulationEngineSubsystem::GetHeadPosition() const
+const FVector& UDEPRECATED_WindowsMixedRealityInputSimulationEngineSubsystem::GetHeadPosition() const
 {
 	FRWScopeLock ReadLock(DataMutex, SLT_ReadOnly);
 	return HeadPosition;
 }
 
-FWindowsMixedRealityInputSimulationHandState* UWindowsMixedRealityInputSimulationEngineSubsystem::GetHandState(EControllerHand Hand)
+FWindowsMixedRealityInputSimulationHandState* UDEPRECATED_WindowsMixedRealityInputSimulationEngineSubsystem::GetHandState(EControllerHand Hand)
 {
 	if (Hand == EControllerHand::Left)
 	{
@@ -94,12 +94,12 @@ FWindowsMixedRealityInputSimulationHandState* UWindowsMixedRealityInputSimulatio
 	return nullptr;
 }
 
-const FWindowsMixedRealityInputSimulationHandState* UWindowsMixedRealityInputSimulationEngineSubsystem::GetHandState(EControllerHand Hand) const
+const FWindowsMixedRealityInputSimulationHandState* UDEPRECATED_WindowsMixedRealityInputSimulationEngineSubsystem::GetHandState(EControllerHand Hand) const
 {
-	return const_cast<UWindowsMixedRealityInputSimulationEngineSubsystem*>(this)->GetHandState(Hand);
+	return const_cast<UDEPRECATED_WindowsMixedRealityInputSimulationEngineSubsystem*>(this)->GetHandState(Hand);
 }
 
-ETrackingStatus UWindowsMixedRealityInputSimulationEngineSubsystem::GetControllerTrackingStatus(EControllerHand Hand) const
+ETrackingStatus UDEPRECATED_WindowsMixedRealityInputSimulationEngineSubsystem::GetControllerTrackingStatus(EControllerHand Hand) const
 {
 	FRWScopeLock ReadLock(DataMutex, SLT_ReadOnly);
 	if (const FWindowsMixedRealityInputSimulationHandState* handState = GetHandState(Hand))
@@ -109,7 +109,7 @@ ETrackingStatus UWindowsMixedRealityInputSimulationEngineSubsystem::GetControlle
 	return ETrackingStatus::NotTracked;
 }
 
-bool UWindowsMixedRealityInputSimulationEngineSubsystem::GetHandJointTransform(EControllerHand Hand, EWMRHandKeypoint Keypoint, FTransform& OutTransform) const
+bool UDEPRECATED_WindowsMixedRealityInputSimulationEngineSubsystem::GetHandJointTransform(EControllerHand Hand, EWMRHandKeypoint Keypoint, FTransform& OutTransform) const
 {
 	FRWScopeLock ReadLock(DataMutex, SLT_ReadOnly);
 	if (const FWindowsMixedRealityInputSimulationHandState* handState = GetHandState(Hand))
@@ -123,7 +123,7 @@ bool UWindowsMixedRealityInputSimulationEngineSubsystem::GetHandJointTransform(E
 	return false;
 }
 
-bool UWindowsMixedRealityInputSimulationEngineSubsystem::GetHandJointRadius(EControllerHand Hand, EWMRHandKeypoint Keypoint, float& OutRadius) const
+bool UDEPRECATED_WindowsMixedRealityInputSimulationEngineSubsystem::GetHandJointRadius(EControllerHand Hand, EWMRHandKeypoint Keypoint, float& OutRadius) const
 {
 	FRWScopeLock ReadLock(DataMutex, SLT_ReadOnly);
 	if (const FWindowsMixedRealityInputSimulationHandState* handState = GetHandState(Hand))
@@ -137,7 +137,7 @@ bool UWindowsMixedRealityInputSimulationEngineSubsystem::GetHandJointRadius(ECon
 	return false;
 }
 
-bool UWindowsMixedRealityInputSimulationEngineSubsystem::HasJointPoses(EControllerHand Hand) const
+bool UDEPRECATED_WindowsMixedRealityInputSimulationEngineSubsystem::HasJointPoses(EControllerHand Hand) const
 {
 	FRWScopeLock ReadLock(DataMutex, SLT_ReadOnly);
 	if (const FWindowsMixedRealityInputSimulationHandState* handState = GetHandState(Hand))
@@ -147,7 +147,7 @@ bool UWindowsMixedRealityInputSimulationEngineSubsystem::HasJointPoses(EControll
 	return false;
 }
 
-bool UWindowsMixedRealityInputSimulationEngineSubsystem::GetPressState(EControllerHand Hand, EHMDInputControllerButtons Button, bool OnlyRegisterClicks, bool& OutPressState) const
+bool UDEPRECATED_WindowsMixedRealityInputSimulationEngineSubsystem::GetPressState(EControllerHand Hand, EHMDInputControllerButtons Button, bool OnlyRegisterClicks, bool& OutPressState) const
 {
 	FRWScopeLock ReadLock(DataMutex, SLT_ReadOnly);
 	if (const FWindowsMixedRealityInputSimulationHandState* handState = GetHandState(Hand))
@@ -169,7 +169,7 @@ bool UWindowsMixedRealityInputSimulationEngineSubsystem::GetPressState(EControll
 	return false;
 }
 
-bool UWindowsMixedRealityInputSimulationEngineSubsystem::GetHandPointerPose(EControllerHand Hand, FWindowsMixedRealityInputSimulationPointerPose& OutPointerPose) const
+bool UDEPRECATED_WindowsMixedRealityInputSimulationEngineSubsystem::GetHandPointerPose(EControllerHand Hand, FWindowsMixedRealityInputSimulationPointerPose& OutPointerPose) const
 {
 	FRWScopeLock ReadLock(DataMutex, SLT_ReadOnly);
 	if (const FWindowsMixedRealityInputSimulationHandState* handState = GetHandState(Hand))
@@ -183,7 +183,7 @@ bool UWindowsMixedRealityInputSimulationEngineSubsystem::GetHandPointerPose(ECon
 	return false;
 }
 
-void UWindowsMixedRealityInputSimulationEngineSubsystem::UpdateSimulatedData(
+void UDEPRECATED_WindowsMixedRealityInputSimulationEngineSubsystem::UpdateSimulatedData(
 	bool HasTracking,
 	const FQuat& NewHeadOrientation,
 	const FVector& NewHeadPosition,
@@ -203,7 +203,7 @@ void UWindowsMixedRealityInputSimulationEngineSubsystem::UpdateSimulatedData(
 	}
 }
 
-void UWindowsMixedRealityInputSimulationEngineSubsystem::UpdateSimulatedHandState(
+void UDEPRECATED_WindowsMixedRealityInputSimulationEngineSubsystem::UpdateSimulatedHandState(
 	EControllerHand Hand,
 	const FWindowsMixedRealityInputSimulationHandState& NewHandState)
 {
@@ -215,7 +215,7 @@ void UWindowsMixedRealityInputSimulationEngineSubsystem::UpdateSimulatedHandStat
 	HandStatePtr->PrevButtonPressed = WasButtonPressed;
 }
 
-bool UWindowsMixedRealityInputSimulationEngineSubsystem::ShouldCreateSubsystem(UObject* Outer) const
+bool UDEPRECATED_WindowsMixedRealityInputSimulationEngineSubsystem::ShouldCreateSubsystem(UObject* Outer) const
 {
 	// Enable input simulation only in the editor
 #if WITH_EDITOR
@@ -225,11 +225,11 @@ bool UWindowsMixedRealityInputSimulationEngineSubsystem::ShouldCreateSubsystem(U
 #endif
 }
 
-void UWindowsMixedRealityInputSimulationEngineSubsystem::Initialize(FSubsystemCollectionBase& Collection)
+void UDEPRECATED_WindowsMixedRealityInputSimulationEngineSubsystem::Initialize(FSubsystemCollectionBase& Collection)
 {
 }
 
-void UWindowsMixedRealityInputSimulationEngineSubsystem::Deinitialize()
+void UDEPRECATED_WindowsMixedRealityInputSimulationEngineSubsystem::Deinitialize()
 {
 }
 
diff --git a/Engine/Plugins/Runtime/WindowsMixedReality/Source/WindowsMixedRealityInputSimulation/Public/WindowsMixedRealityInputSimulationEngineSubsystem.h b/Engine/Plugins/Runtime/WindowsMixedReality/Source/WindowsMixedRealityInputSimulation/Public/WindowsMixedRealityInputSimulationEngineSubsystem.h
--- a/Engine/Plugins/Runtime/WindowsMixedReality/Source/WindowsMixedRealityInputSimulation/Public/WindowsMixedRealityInputSimulationEngineSubsystem.h
+++ b/Engine/Plugins/Runtime/WindowsMixedReality/Source/WindowsMixedRealityInputSimulation/Public/WindowsMixedRealityInputSimulationEngineSubsystem.h
@@ -42,8 +42,8 @@ struct WINDOWSMIXEDREALITYINPUTSIMULATION_API FWindowsMixedRealityInputSimulatio
 };
 
 /** Engine subsystem that stores input simulation data for access by the XR device. */
-UCLASS(ClassGroup = WindowsMixedReality)
-class WINDOWSMIXEDREALITYINPUTSIMULATION_API UWindowsMixedRealityInputSimulationEngineSubsystem
+UCLASS(ClassGroup = WindowsMixedReality, Deprecated, meta = (DeprecationMessage = "WMR has been deprecated, please use the Microsoft OpenXR plugin from the Marketplace."))
+class WINDOWSMIXEDREALITYINPUTSIMULATION_API UDEPRECATED_WindowsMixedRealityInputSimulationEngineSubsystem
 	: public UEngineSubsystem
 {
 	GENERATED_BODY()
@@ -59,7 +59,7 @@ public:
 	/** Utility function to ensure input simulation is only used when enabled.
 	 * Returns the engine subsystem if it exists and if input simulation is enabled.
 	 */
-	static UWindowsMixedRealityInputSimulationEngineSubsystem* GetInputSimulationIfEnabled();
+	static UDEPRECATED_WindowsMixedRealityInputSimulationEngineSubsystem* GetInputSimulationIfEnabled();
 		
 	/** True if the HMD has a valid head position vector. */
 	bool HasPositionalTracking() const;
diff --git a/Engine/Plugins/Runtime/WindowsMixedReality/WindowsMixedReality.uplugin b/Engine/Plugins/Runtime/WindowsMixedReality/WindowsMixedReality.uplugin
--- a/Engine/Plugins/Runtime/WindowsMixedReality/WindowsMixedReality.uplugin
+++ b/Engine/Plugins/Runtime/WindowsMixedReality/WindowsMixedReality.uplugin
@@ -3,7 +3,7 @@
 	"Version": 1,
 	"VersionName": "1.0",
 	"FriendlyName": "Microsoft Windows Mixed Reality",
-	"Description": "Support for Microsoft Windows Mixed Reality compatible devices.",
+	"Description": "Support for Microsoft Windows Mixed Reality compatible devices. This plugin is deprecated and will be removed in a future engine release. Please use the Microsoft OpenXR plugin from the Marketplace instead.",
 	"Category": "Virtual Reality",
 	"CreatedBy": "Microsoft Corporation",
 	"CreatedByURL": "http://www.microsoft.com",
