Fix startup error in editor if WMR is enabled and the windows build is later than 2004 (e.g. 20H1, 20H2 etc). This is an interim fix, since the WMR path is tagged as deprecated in 5.0 and will be removed in a future version.

diff --git a/Engine/Plugins/Runtime/WindowsMixedReality/Source/WindowsMixedRealityHMD/Private/WindowsMixedRealityInteropLoader.cpp b/Engine/Plugins/Runtime/WindowsMixedReality/Source/WindowsMixedRealityHMD/Private/WindowsMixedRealityInteropLoader.cpp
--- a/Engine/Plugins/Runtime/WindowsMixedReality/Source/WindowsMixedRealityHMD/Private/WindowsMixedRealityInteropLoader.cpp
+++ b/Engine/Plugins/Runtime/WindowsMixedReality/Source/WindowsMixedRealityHMD/Private/WindowsMixedRealityInteropLoader.cpp
@@ -46,6 +46,13 @@ namespace WindowsMixedReality
 				OSSubVersionLabel = OSSubVersionLabel.Mid(Begin, End - Begin);
 				bHasSupportedWindowsBuild = FCString::Atoi(*OSSubVersionLabel) >= MIN_WIN_10_VERSION_FOR_WMR;
 			}
+			else
+			{
+				// Hack: Later versions of windows change the format of this sub version label string, so if we hit this case, we should be on a supported version.
+				// This isn't an ideal method, but ran this by Microsoft and for now I think this is a reasonable fix. The WMR path is deprecated in the UE 5.0
+				// release and will be removed completely in a later release.
+				bHasSupportedWindowsBuild = true;
+			}
 		}
 
 		if (bHasSupportedWindowsBuild && bHasSupportedWindowsVersion)
