Fix build break.
-WindowsMixedReality.uplugin needs to be reverted, this is referencing a module that still exists, not the plugin that was removed.

diff --git a/Engine/Plugins/Runtime/WindowsMixedReality/WindowsMixedReality.uplugin b/Engine/Plugins/Runtime/WindowsMixedReality/WindowsMixedReality.uplugin
--- a/Engine/Plugins/Runtime/WindowsMixedReality/WindowsMixedReality.uplugin
+++ b/Engine/Plugins/Runtime/WindowsMixedReality/WindowsMixedReality.uplugin
@@ -88,6 +88,15 @@
 			"PlatformAllowList": [
 				"Win64"
 			]
+		},
+		{
+			"Name": "HoloLensAR",
+			"Type": "Runtime",
+			"LoadingPhase": "PostConfigInit",
+			"PlatformAllowList": [
+				"HoloLens",
+				"Win64"
+			]
 		}
 	],
 	"Plugins": [
