Restoring D3D11 includes to WindowsMixedRealityRHI.
Fixing WindowsMixedRealityHMD errors.

diff --git a/Engine/Plugins/Runtime/WindowsMixedReality/Source/WindowsMixedRealityHMD/Private/WindowsMixedRealityHMD.cpp b/Engine/Plugins/Runtime/WindowsMixedReality/Source/WindowsMixedRealityHMD/Private/WindowsMixedRealityHMD.cpp
--- a/Engine/Plugins/Runtime/WindowsMixedReality/Source/WindowsMixedRealityHMD/Private/WindowsMixedRealityHMD.cpp
+++ b/Engine/Plugins/Runtime/WindowsMixedReality/Source/WindowsMixedRealityHMD/Private/WindowsMixedRealityHMD.cpp
@@ -719,7 +719,7 @@ namespace WindowsMixedReality
 		if (bIsMobileMultiViewEnabled && !HMD->IsThirdCameraActive())
 		{
 			ID3D11Texture2D* Texture = HMD->GetBackBufferTexture();
-			FTexture2DArrayRHIRef BackBuffer = GD3D11RHI->RHICreateTexture2DArrayFromResource(PF_B8G8R8A8, TexCreate_RenderTargetable | TexCreate_ShaderResource, FClearValueBinding::None, Texture);
+			FTexture2DArrayRHIRef BackBuffer = GetID3D11DynamicRHI()->RHICreateTexture2DArrayFromResource(PF_B8G8R8A8, TexCreate_RenderTargetable | TexCreate_ShaderResource, FClearValueBinding::None, Texture);
 			GDynamicRHI->RHIAliasTextureResources((FTextureRHIRef&)CurrentBackBuffer, (FTextureRHIRef&)BackBuffer);
 		}
 #endif
diff --git a/Engine/Plugins/Runtime/WindowsMixedReality/Source/WindowsMixedRealityRHI/Private/WindowsMixedRealityPrecompiled.h b/Engine/Plugins/Runtime/WindowsMixedReality/Source/WindowsMixedRealityRHI/Private/WindowsMixedRealityPrecompiled.h
--- a/Engine/Plugins/Runtime/WindowsMixedReality/Source/WindowsMixedRealityRHI/Private/WindowsMixedRealityPrecompiled.h
+++ b/Engine/Plugins/Runtime/WindowsMixedReality/Source/WindowsMixedRealityRHI/Private/WindowsMixedRealityPrecompiled.h
@@ -22,4 +22,8 @@
 DECLARE_LOG_CATEGORY_EXTERN(LogWMRRHI, Log, All);
 
 #include "SceneRendering.h"
-#include "ID3D11DynamicRHI.h"
+
+#pragma warning(disable:4668)  
+#include "D3D11RHIPrivate.h"
+#pragma warning(default:4668)
+#include "D3D11Util.h"
\ No newline at end of file
diff --git a/Engine/Plugins/Runtime/WindowsMixedReality/Source/WindowsMixedRealityRHI/WindowsMixedRealityRHI.Build.cs b/Engine/Plugins/Runtime/WindowsMixedReality/Source/WindowsMixedRealityRHI/WindowsMixedRealityRHI.Build.cs
--- a/Engine/Plugins/Runtime/WindowsMixedReality/Source/WindowsMixedRealityRHI/WindowsMixedRealityRHI.Build.cs
+++ b/Engine/Plugins/Runtime/WindowsMixedReality/Source/WindowsMixedRealityRHI/WindowsMixedRealityRHI.Build.cs
@@ -53,6 +53,7 @@ namespace UnrealBuildTool.Rules
 						"Projects",
                         "HeadMountedDisplay",
                         "D3D11RHI",
+						"RHICore",
                         "Slate",
                         "SlateCore",
                         "MRMesh",
@@ -62,6 +63,12 @@ namespace UnrealBuildTool.Rules
 
                 AddEngineThirdPartyPrivateStaticDependencies(Target, "WindowsMixedRealityInterop");
 
+				if (Target.Platform != UnrealTargetPlatform.HoloLens)
+				{
+					AddEngineThirdPartyPrivateStaticDependencies(Target, "NVAftermath");
+					AddEngineThirdPartyPrivateStaticDependencies(Target, "IntelExtensionsFramework");
+				}
+
 				if (Target.bBuildEditor == true)
                 {
 					PrivateDependencyModuleNames.Add("EditorFramework");
@@ -74,9 +81,20 @@ namespace UnrealBuildTool.Rules
                     new string[]
                     {
                     "WindowsMixedRealityRHI/Private",
+                    "../../../../Source/Runtime/Windows/D3D11RHI/Private",
                     "../../../../Source/Runtime/Renderer/Private",
                     });
 
+                //TODO: needed?
+                if (Target.Platform == UnrealTargetPlatform.Win64)
+                {
+                    PrivateIncludePaths.Add("../../../../Source/Runtime/Windows/D3D11RHI/Private/Windows");
+                }
+                else if (Target.Platform == UnrealTargetPlatform.HoloLens)
+                {
+                    PrivateIncludePaths.Add("../../../../Source/Runtime/Windows/D3D11RHI/Private/HoloLens");
+                }
+
                 PCHUsage = PCHUsageMode.NoSharedPCHs;
                 PrivatePCHHeaderFile = "Private/WindowsMixedRealityPrecompiled.h";
             }
