UE-146001 //UE5/Main - Compile UnrealGame HoloLens - WindowsMixedRealityRHI.h error: 'FD3D11Texture2D': undeclared identifier
-Updated a couple type names due to rhi refactor.

diff --git a/Engine/Plugins/Runtime/WindowsMixedReality/Source/WindowsMixedRealityRHI/Private/WindowsMixedRealityRHI.cpp b/Engine/Plugins/Runtime/WindowsMixedReality/Source/WindowsMixedRealityRHI/Private/WindowsMixedRealityRHI.cpp
--- a/Engine/Plugins/Runtime/WindowsMixedReality/Source/WindowsMixedRealityRHI/Private/WindowsMixedRealityRHI.cpp
+++ b/Engine/Plugins/Runtime/WindowsMixedReality/Source/WindowsMixedRealityRHI/Private/WindowsMixedRealityRHI.cpp
@@ -295,8 +295,8 @@ FWindowsMixedRealityViewport::FWindowsMixedRealityViewport(FD3D11DynamicRHI* InD
 
 void FWindowsMixedRealityViewport::UpdateBackBuffer()
 {
-	FTexture2DRHIRef OutShaderResourceTexture;
-	FTexture2DRHIRef OutTargetableTexture;
+	FTextureRHIRef OutShaderResourceTexture;
+	FTextureRHIRef OutTargetableTexture;
 
 	BackBuffer = nullptr;
 	if (GEngine && GEngine->StereoRenderingDevice && GEngine->StereoRenderingDevice->GetRenderTargetManager())
@@ -306,7 +306,7 @@ void FWindowsMixedRealityViewport::UpdateBackBuffer()
 			->GetRenderTargetManager()
 			->AllocateRenderTargetTexture(0, SizeX, SizeY, PixelFormat, 1, TexCreate_None, TexCreate_RenderTargetable, OutTargetableTexture, OutShaderResourceTexture))
 		{
-			BackBuffer = (FD3D11Texture2D*)(OutTargetableTexture.GetReference());
+			BackBuffer = (FD3D11Texture*)(OutTargetableTexture.GetReference());
 		}
 	}
 
diff --git a/Engine/Plugins/Runtime/WindowsMixedReality/Source/WindowsMixedRealityRHI/Private/WindowsMixedRealityRHI.h b/Engine/Plugins/Runtime/WindowsMixedReality/Source/WindowsMixedRealityRHI/Private/WindowsMixedRealityRHI.h
--- a/Engine/Plugins/Runtime/WindowsMixedReality/Source/WindowsMixedRealityRHI/Private/WindowsMixedRealityRHI.h
+++ b/Engine/Plugins/Runtime/WindowsMixedReality/Source/WindowsMixedRealityRHI/Private/WindowsMixedRealityRHI.h
@@ -78,5 +78,5 @@ public:
 	void Resize(uint32 InSizeX, uint32 InSizeY, bool bInIsFullscreen, EPixelFormat PreferredPixelFormat);
 
 private:
-	TRefCountPtr<FD3D11Texture2D> OffscreenBackBuffer;
+	TRefCountPtr<FD3D11Texture> OffscreenBackBuffer;
 };
