WindowsMixedReality UE5.1 fixes
-Initializing a struct.
-Pipeline pixel density because it is accessed from multiple threads.

diff --git a/Engine/Plugins/Runtime/WindowsMixedReality/Source/WindowsMixedRealityHMD/Private/WindowsMixedRealityHMD.cpp b/Engine/Plugins/Runtime/WindowsMixedReality/Source/WindowsMixedRealityHMD/Private/WindowsMixedRealityHMD.cpp
--- a/Engine/Plugins/Runtime/WindowsMixedReality/Source/WindowsMixedRealityHMD/Private/WindowsMixedRealityHMD.cpp
+++ b/Engine/Plugins/Runtime/WindowsMixedReality/Source/WindowsMixedRealityHMD/Private/WindowsMixedRealityHMD.cpp
@@ -703,6 +703,8 @@ namespace WindowsMixedReality
 		}
 		InitTrackingFrame();
 
+		ScreenScalePercentage_RenderThread = ScreenScalePercentage;
+
 		if (SpectatorScreenController)
 		{
 			SpectatorScreenController->UpdateSpectatorScreenMode_RenderThread();
@@ -1521,8 +1523,20 @@ namespace WindowsMixedReality
 	//TODO: Spelling is intentional, overridden from IHeadMountedDisplay.h
 	float FWindowsMixedRealityHMD::GetPixelDenity() const
 	{
-		check(IsInGameThread());
-		return ScreenScalePercentage;
+		if (IsInGameThread())
+		{
+			return ScreenScalePercentage;
+
+		}
+		else if (IsInRenderingThread())
+		{
+			return ScreenScalePercentage_RenderThread;
+		}
+		else
+		{
+			check(false);
+			return ScreenScalePercentage_RenderThread; // Not the game or render thread, so ??? Maybe rhi?
+		}
 	}
 
 	void FWindowsMixedRealityHMD::SetPixelDensity(const float NewDensity)
@@ -2019,7 +2033,6 @@ namespace WindowsMixedReality
 		: FHeadMountedDisplayBase(InARSystem)
 		, FHMDSceneViewExtension(AutoRegister)
 		, HMD(InHMD)
-		, ScreenScalePercentage(1.0f)
 		, mCustomPresent(nullptr)
 	{
 		static const FName RendererModuleName("Renderer");
diff --git a/Engine/Plugins/Runtime/WindowsMixedReality/Source/WindowsMixedRealityHMD/Private/WindowsMixedRealityHMD.h b/Engine/Plugins/Runtime/WindowsMixedReality/Source/WindowsMixedRealityHMD/Private/WindowsMixedRealityHMD.h
--- a/Engine/Plugins/Runtime/WindowsMixedReality/Source/WindowsMixedRealityHMD/Private/WindowsMixedRealityHMD.h
+++ b/Engine/Plugins/Runtime/WindowsMixedReality/Source/WindowsMixedRealityHMD/Private/WindowsMixedRealityHMD.h
@@ -220,6 +220,7 @@ namespace WindowsMixedReality
 		bool bIsMobileMultiViewEnabled = false;
 
 		float ScreenScalePercentage = 1.0f;
+		float ScreenScalePercentage_RenderThread = 1.0f;
 		float CachedWorldToMetersScale = 100.0f;
 
 		TRefCountPtr<ID3D11Device> D3D11Device = nullptr;
diff --git a/Engine/Plugins/Runtime/WindowsMixedReality/Source/WindowsMixedRealityHMD/Public/WindowsMixedRealityFunctionLibrary.h b/Engine/Plugins/Runtime/WindowsMixedReality/Source/WindowsMixedRealityHMD/Public/WindowsMixedRealityFunctionLibrary.h
--- a/Engine/Plugins/Runtime/WindowsMixedReality/Source/WindowsMixedRealityHMD/Public/WindowsMixedRealityFunctionLibrary.h
+++ b/Engine/Plugins/Runtime/WindowsMixedReality/Source/WindowsMixedRealityHMD/Public/WindowsMixedRealityFunctionLibrary.h
@@ -43,6 +43,16 @@ struct FPointerPoseInfo
 {
 	GENERATED_USTRUCT_BODY()
 
+public:
+
+	FPointerPoseInfo()
+		: Origin(FVector::ZeroVector)
+		, Direction(FVector::ZeroVector)
+		, Up(FVector::ZeroVector)
+		, TrackingStatus(EHMDTrackingStatus::NotTracked)
+	{
+	}
+
 	UPROPERTY(EditAnywhere, BlueprintReadWrite, Category = "WindowsMixedRealityHMD")
 	FVector Origin;
 
