UE-145518 Connecting HoloLens 2 remotely to editor crashes the engine with FPS template project
-Added a check that the RHI is DX11, because only that is supported.  If not we log a warning and do not attempt to connect.

diff --git a/Engine/Plugins/Runtime/WindowsMixedReality/Source/WindowsMixedRealityHMD/Private/WindowsMixedRealityHMD.cpp b/Engine/Plugins/Runtime/WindowsMixedReality/Source/WindowsMixedRealityHMD/Private/WindowsMixedRealityHMD.cpp
--- a/Engine/Plugins/Runtime/WindowsMixedReality/Source/WindowsMixedRealityHMD/Private/WindowsMixedRealityHMD.cpp
+++ b/Engine/Plugins/Runtime/WindowsMixedReality/Source/WindowsMixedRealityHMD/Private/WindowsMixedRealityHMD.cpp
@@ -2512,6 +2512,13 @@ namespace WindowsMixedReality
 	{
 		EXRDeviceConnectionResult::Type Result = EXRDeviceConnectionResult::Success;
 
+		if (TCString<TCHAR>::Stricmp(GDynamicRHI->GetName(), TEXT("D3D11")) != 0)
+		{
+			UE_LOG(LogWmrHmd, Warning, TEXT("FWindowsMixedRealityHMD::ConnectToRemoteHoloLens() not connecting because the RHI is %s but only D3D11 is supported."), GDynamicRHI->GetName());
+			Result = EXRDeviceConnectionResult::MiscFailure;
+			return Result;
+		}
+
 		D3D11Device = InternalGetD3D11Device();
 
 #  if WITH_WINDOWS_MIXED_REALITY
@@ -2523,7 +2530,7 @@ namespace WindowsMixedReality
 		HMD->SetLogCallback(WindowsMixedRealityHMD::LogForInterop);
 		HMD->ConnectToRemoteHoloLens(D3D11Device.GetReference(), ip, bitrate, Port, listen);
 #  else
-		UE_LOG(LogWmrHmd, Log, TEXT("FWindowsMixedRealityHMD::ConnectToRemoteHoloLens() is doing nothing because !WITH_WINDOWS_MIXED_REALITY."));
+		UE_LOG(LogWmrHmd, Warning, TEXT("FWindowsMixedRealityHMD::ConnectToRemoteHoloLens() is doing nothing because !WITH_WINDOWS_MIXED_REALITY."));
 #  endif
 
 		TWeakPtr<SViewport> ActiveViewport;
