PR #9129: Add Windows 11 to supported WMR Windows Versions (Contributed by [REDACTED_PERSON])

ROB: Legacy fix for WMR. Kinda clumsy, but fits in with current pattern + what is done in FMicrosoftSpatialSoundModule::StartupModule()

diff --git a/Engine/Plugins/Runtime/WindowsMixedReality/Source/WindowsMixedRealityHMD/Private/WindowsMixedRealityInteropLoader.cpp b/Engine/Plugins/Runtime/WindowsMixedReality/Source/WindowsMixedRealityHMD/Private/WindowsMixedRealityInteropLoader.cpp
--- a/Engine/Plugins/Runtime/WindowsMixedReality/Source/WindowsMixedRealityHMD/Private/WindowsMixedRealityInteropLoader.cpp
+++ b/Engine/Plugins/Runtime/WindowsMixedReality/Source/WindowsMixedRealityHMD/Private/WindowsMixedRealityInteropLoader.cpp
@@ -25,7 +25,9 @@ namespace WindowsMixedReality
 		FPlatformMisc::GetOSVersions(OSVersionLabel, OSSubVersionLabel);
 		// GetOSVersion returns the Win10 release version in the OSVersion rather than the OSSubVersion, so parse it out ourselves
 		OSSubVersionLabel = OSVersionLabel;
-		bool bHasSupportedWindowsVersion = OSSubVersionLabel.StartsWith("Windows 10") || OSSubVersionLabel.StartsWith("Windows Server 2019");
+		bool bHasSupportedWindowsVersion = OSSubVersionLabel.StartsWith("Windows 10") || 
+			OSSubVersionLabel.StartsWith("Windows Server 2019") ||
+			OSSubVersionLabel.StartsWith("Windows 11");
 		bool bHasSupportedWindowsBuild = false;
 
 		// If we can't find Win10 version, check for Windows Server equivalent
@@ -122,4 +124,4 @@ namespace WindowsMixedReality
 #endif
 		return HMD;
 	}
-}
\ No newline at end of file
+}
