Changing uses of RHIGetNativeDevice and GetNativeResource to use the per-RHI interfaces. XR Edition.

diff --git a/Engine/Plugins/Media/HLMedia/Source/HLMedia/Private/HLMediaPlayer.cpp b/Engine/Plugins/Media/HLMedia/Source/HLMedia/Private/HLMediaPlayer.cpp
--- a/Engine/Plugins/Media/HLMedia/Source/HLMedia/Private/HLMediaPlayer.cpp
+++ b/Engine/Plugins/Media/HLMedia/Source/HLMedia/Private/HLMediaPlayer.cpp
@@ -381,14 +381,14 @@ bool FHLMediaPlayer::InitializePlayer(const TSharedPtr<FArchive, ESPMode::Thread
     {
         UE_LOG(LogHLMediaPlayer, Verbose, TEXT("HLMediaPlayer %p: FHLMediaPlayer::InitializePlayer::Async()"), this);
 
-        ID3D11Device* Device = static_cast<ID3D11Device*>(GDynamicRHI->RHIGetNativeDevice());
+		ID3D11Device* Device = GetID3D11DynamicRHI()->RHIGetDevice();
 		checkf(Device, TEXT("No available D3D11Device"));
 
 		if (RHIGetInterfaceType() == ERHIInterfaceType::D3D12)
 		{
 			ID3D11Device* PrevDevice = Device;
-			ID3D12Device* D3d12Device = static_cast<ID3D12Device*>(GDynamicRHI->RHIGetNativeDevice());
-			void* CommandQueue = GDynamicRHI->RHIGetNativeGraphicsQueue();
+			ID3D12Device* D3d12Device = GetID3D12DynamicRHI()->RHIGetDevice(0);
+			ID3D12CommandQueue* CommandQueue = GetID3D12DynamicRHI()->RHIGetCommandQueue();
 
 			ID3D11DeviceContext* D3d11DeviceContext;
 			if (FAILED(D3D11On12CreateDevice(
diff --git a/Engine/Plugins/Runtime/Oculus/OculusVR/Source/OculusHMD/Private/OculusHMD_CustomPresent_D3D11.cpp b/Engine/Plugins/Runtime/Oculus/OculusVR/Source/OculusHMD/Private/OculusHMD_CustomPresent_D3D11.cpp
--- a/Engine/Plugins/Runtime/Oculus/OculusVR/Source/OculusHMD/Private/OculusHMD_CustomPresent_D3D11.cpp
+++ b/Engine/Plugins/Runtime/Oculus/OculusVR/Source/OculusHMD/Private/OculusHMD_CustomPresent_D3D11.cpp
@@ -53,25 +53,12 @@ bool FD3D11CustomPresent::IsUsingCorrectDisplayAdapter() const
 
 	if (OVRP_SUCCESS(FOculusHMDModule::GetPluginWrapper().GetDisplayAdapterId2(&luid)) && luid)
 	{
-		TRefCountPtr<ID3D11Device> D3D11Device;
+		IDXGIAdapter* Adapter = GetID3D11DynamicRHI()->RHIGetAdapter();
+		DXGI_ADAPTER_DESC AdapterDesc{};
 
-		ExecuteOnRenderThread([&D3D11Device]()
+		if (Adapter && SUCCEEDED(Adapter->GetDesc(&AdapterDesc)))
 		{
-			D3D11Device = (ID3D11Device*) RHIGetNativeDevice();
-		});
-
-		if (D3D11Device)
-		{
-			TRefCountPtr<IDXGIDevice> DXGIDevice;
-			TRefCountPtr<IDXGIAdapter> DXGIAdapter;
-			DXGI_ADAPTER_DESC DXGIAdapterDesc;
-
-			if (SUCCEEDED(D3D11Device->QueryInterface(__uuidof(IDXGIDevice), (void**) DXGIDevice.GetInitReference())) &&
-				SUCCEEDED(DXGIDevice->GetAdapter(DXGIAdapter.GetInitReference())) &&
-				SUCCEEDED(DXGIAdapter->GetDesc(&DXGIAdapterDesc)))
-			{
-				return !FMemory::Memcmp(luid, &DXGIAdapterDesc.AdapterLuid, sizeof(LUID));
-			}
+			return !FMemory::Memcmp(luid, &AdapterDesc.AdapterLuid, sizeof(LUID));
 		}
 	}
 
diff --git a/Engine/Plugins/Runtime/Oculus/OculusVR/Source/OculusHMD/Private/OculusHMD_CustomPresent_D3D12.cpp b/Engine/Plugins/Runtime/Oculus/OculusVR/Source/OculusHMD/Private/OculusHMD_CustomPresent_D3D12.cpp
--- a/Engine/Plugins/Runtime/Oculus/OculusVR/Source/OculusHMD/Private/OculusHMD_CustomPresent_D3D12.cpp
+++ b/Engine/Plugins/Runtime/Oculus/OculusVR/Source/OculusHMD/Private/OculusHMD_CustomPresent_D3D12.cpp
@@ -53,13 +53,7 @@ bool FD3D12CustomPresent::IsUsingCorrectDisplayAdapter() const
 
 	if (OVRP_SUCCESS(FOculusHMDModule::GetPluginWrapper().GetDisplayAdapterId2(&luid)) && luid)
 	{
-		TRefCountPtr<ID3D12Device> D3DDevice;
-
-		ExecuteOnRenderThread([&D3DDevice]()
-		{
-			D3DDevice = (ID3D12Device*) RHIGetNativeDevice();
-		});
-
+		ID3D12Device* D3DDevice = GetID3D12DynamicRHI()->RHIGetDevice(0);
 		if (D3DDevice)
 		{
 			LUID AdapterLuid = D3DDevice->GetAdapterLuid();
diff --git a/Engine/Plugins/Runtime/WindowsMixedReality/Source/HoloLensAR/HoloLensAR.Build.cs b/Engine/Plugins/Runtime/WindowsMixedReality/Source/HoloLensAR/HoloLensAR.Build.cs
--- a/Engine/Plugins/Runtime/WindowsMixedReality/Source/HoloLensAR/HoloLensAR.Build.cs
+++ b/Engine/Plugins/Runtime/WindowsMixedReality/Source/HoloLensAR/HoloLensAR.Build.cs
@@ -47,6 +47,9 @@ public class HoloLensAR : ModuleRules
             {
                 RuntimeDependencies.Add(Dll);
             }
+
+			PrivateDependencyModuleNames.Add("D3D11RHI");
+			AddEngineThirdPartyPrivateStaticDependencies(Target, "DX11");
         }
         else if (Target.Platform == UnrealTargetPlatform.HoloLens)
         {
@@ -66,7 +69,10 @@ public class HoloLensAR : ModuleRules
             {
                 PublicDefinitions.Add("WITH_SCENE_UNDERSTANDING=0");
             }
-        }
+
+			PrivateDependencyModuleNames.Add("D3D11RHI");
+			AddEngineThirdPartyPrivateStaticDependencies(Target, "DX11");
+		}
 
 		AddEngineThirdPartyPrivateStaticDependencies(Target, "WindowsMixedRealityInterop");
 
diff --git a/Engine/Plugins/Runtime/WindowsMixedReality/Source/HoloLensAR/Private/HoloLensCameraImageTexture.cpp b/Engine/Plugins/Runtime/WindowsMixedReality/Source/HoloLensAR/Private/HoloLensCameraImageTexture.cpp
--- a/Engine/Plugins/Runtime/WindowsMixedReality/Source/HoloLensAR/Private/HoloLensCameraImageTexture.cpp
+++ b/Engine/Plugins/Runtime/WindowsMixedReality/Source/HoloLensAR/Private/HoloLensCameraImageTexture.cpp
@@ -10,10 +10,7 @@
 #include "SceneUtils.h"
 
 #if SUPPORTS_WINDOWS_MIXED_REALITY_AR
-	#include "Windows/AllowWindowsPlatformTypes.h"
-	#include <Windows.h>
-	#include <d3d11.h>
-	#include "Windows/HideWindowsPlatformTypes.h"
+#include "ID3D11DynamicRHI.h"
 #endif
 
 struct FHoloLensCameraImageConversionVertex
@@ -193,9 +190,8 @@ public:
 		if (CameraImageHandle != INVALID_HANDLE_VALUE)
 		{
 			// Open the shared texture from the HoloLens camera on Unreal's d3d device.
-			ID3D11Device* D3D11Device = static_cast<ID3D11Device*>(GDynamicRHI->RHIGetNativeDevice());
-			TComPtr<ID3D11DeviceContext> D3D11DeviceContext = nullptr;
-			D3D11Device->GetImmediateContext(&D3D11DeviceContext);
+			ID3D11Device* D3D11Device = GetID3D11DynamicRHI()->RHIGetDevice();
+			TComPtr<ID3D11DeviceContext> D3D11DeviceContext = GetID3D11DynamicRHI()->RHIGetDeviceContext();
 			if (D3D11DeviceContext == nullptr)
 			{
 				CloseHandle(CameraImageHandle);
@@ -313,7 +309,7 @@ private:
 			return false;
 		}
 		// Get the underlying interface for the texture we are copying to
-		TComPtr<ID3D11Resource> CopyTexture = reinterpret_cast<ID3D11Resource*>(CopyTextureRef->GetNativeResource());
+		TComPtr<ID3D11Resource> CopyTexture = GetID3D11DynamicRHI()->RHIGetResource(CopyTextureRef);
 		if (CopyTexture == nullptr)
 		{
 			return false;
diff --git a/Engine/Plugins/Runtime/WindowsMixedReality/Source/WindowsMixedRealityHMD/Private/WindowsMixedRealityCustomPresent.h b/Engine/Plugins/Runtime/WindowsMixedReality/Source/WindowsMixedRealityHMD/Private/WindowsMixedRealityCustomPresent.h
--- a/Engine/Plugins/Runtime/WindowsMixedReality/Source/WindowsMixedRealityHMD/Private/WindowsMixedRealityCustomPresent.h
+++ b/Engine/Plugins/Runtime/WindowsMixedReality/Source/WindowsMixedRealityHMD/Private/WindowsMixedRealityCustomPresent.h
@@ -93,7 +93,7 @@ namespace WindowsMixedReality
 				return;
 			}
 
-			ViewportTexture = (ID3D11Texture2D*)RT->GetNativeResource();
+			ViewportTexture = (ID3D11Texture2D*)GetID3D11DynamicRHI()->RHIGetResource(RT);
 		}
 
 		void SetDepthTexture(ID3D11Texture2D* depthTexture)
diff --git a/Engine/Plugins/Runtime/WindowsMixedReality/Source/WindowsMixedRealityHMD/Private/WindowsMixedRealityHMD.cpp b/Engine/Plugins/Runtime/WindowsMixedReality/Source/WindowsMixedRealityHMD/Private/WindowsMixedRealityHMD.cpp
--- a/Engine/Plugins/Runtime/WindowsMixedReality/Source/WindowsMixedRealityHMD/Private/WindowsMixedRealityHMD.cpp
+++ b/Engine/Plugins/Runtime/WindowsMixedReality/Source/WindowsMixedRealityHMD/Private/WindowsMixedRealityHMD.cpp
@@ -322,7 +322,7 @@ namespace WindowsMixedReality
 			FWindowsMixedRealityHMD* Self = this;
 			ENQUEUE_RENDER_COMMAND(InternalGetD3D11DeviceCmd)([Self](FRHICommandListImmediate& RHICmdList)
 			{
-				Self->D3D11Device = (ID3D11Device*)RHIGetNativeDevice();
+				Self->D3D11Device = GetID3D11DynamicRHI()->RHIGetDevice();
 			});
 
 			FlushRenderingCommands();
@@ -1571,7 +1571,7 @@ namespace WindowsMixedReality
 		{
 			if (FRHITexture* depthTargetableTexture = GetSceneTextureExtracts().GetDepthTexture())
 			{
-				ID3D11Texture2D* Texture = static_cast<ID3D11Texture2D*>(depthTargetableTexture->GetNativeResource());
+				ID3D11Texture2D* Texture = static_cast<ID3D11Texture2D*>(GetID3D11DynamicRHI()->RHIGetResource(depthTargetableTexture));
 				if (Texture != nullptr)
 				{
 					if (mCustomPresent != nullptr)
@@ -1801,7 +1801,7 @@ namespace WindowsMixedReality
 		// Update depth texture to match format Windows Mixed Reality platform is expecting.
 		FRHITexture2D* depthFRHITexture = CurrentDepthBuffer.GetReference()->GetTexture2D();
 
-		if (depthFRHITexture == nullptr || depthFRHITexture->GetNativeResource() == nullptr)
+		if (depthFRHITexture == nullptr || GetID3D11DynamicRHI()->RHIGetResource(depthFRHITexture) == nullptr)
 		{
 			return;
 		}
@@ -1869,7 +1869,7 @@ namespace WindowsMixedReality
 		}
 		RHICmdList.EndRenderPass();
 
-		ID3D11Device* device = (ID3D11Device*)RHIGetNativeDevice();
+		ID3D11Device* device = GetID3D11DynamicRHI()->RHIGetDevice();
 		if (device == nullptr)
 		{
 			return;
@@ -1920,7 +1920,7 @@ namespace WindowsMixedReality
 		device->GetImmediateContext(&context);
 
 		_StereoCopy(context, ScreenScalePercentage,
-			(ID3D11Texture2D*)remappedDepthTexture->GetNativeResource(),
+			(ID3D11Texture2D*)GetID3D11DynamicRHI()->RHIGetResource(remappedDepthTexture),
 			stereoDepthTexture);
 
 		if (mCustomPresent != nullptr)
@@ -1930,7 +1930,7 @@ namespace WindowsMixedReality
 		// Third camera depth
 		if (HMD->IsThirdCameraActive() && remappedDepthTexture != nullptr && monoDepthTexture != nullptr)
 		{
-			ID3D11Texture2D* sourceDepth = (ID3D11Texture2D*)remappedDepthTexture->GetNativeResource();
+			ID3D11Texture2D* sourceDepth = (ID3D11Texture2D*)GetID3D11DynamicRHI()->RHIGetResource(remappedDepthTexture);
 			int w, h;
 			HMD->GetThirdCameraDimensions(w, h);
 
diff --git a/Engine/Source/Runtime/Windows/D3D11RHI/Private/D3D11RHI.cpp b/Engine/Source/Runtime/Windows/D3D11RHI/Private/D3D11RHI.cpp
--- a/Engine/Source/Runtime/Windows/D3D11RHI/Private/D3D11RHI.cpp
+++ b/Engine/Source/Runtime/Windows/D3D11RHI/Private/D3D11RHI.cpp
@@ -772,6 +772,11 @@ ID3D11DeviceContext* FD3D11DynamicRHI::RHIGetDeviceContext() const
 	return GetDeviceContext();
 }
 
+IDXGIAdapter* FD3D11DynamicRHI::RHIGetAdapter() const
+{
+	return GetAdapter().DXGIAdapter;
+}
+
 IDXGISwapChain* FD3D11DynamicRHI::RHIGetSwapChain(FRHIViewport* InViewport) const
 {
 	FD3D11Viewport* Viewport = static_cast<FD3D11Viewport*>(InViewport);
diff --git a/Engine/Source/Runtime/Windows/D3D11RHI/Private/D3D11RHIPrivate.h b/Engine/Source/Runtime/Windows/D3D11RHI/Private/D3D11RHIPrivate.h
--- a/Engine/Source/Runtime/Windows/D3D11RHI/Private/D3D11RHIPrivate.h
+++ b/Engine/Source/Runtime/Windows/D3D11RHI/Private/D3D11RHIPrivate.h
@@ -663,6 +663,7 @@ public:
 	// ID3D11DynamicRHI interface
 	virtual ID3D11Device*         RHIGetDevice() const final override;
 	virtual ID3D11DeviceContext*  RHIGetDeviceContext() const final override;
+	virtual IDXGIAdapter*         RHIGetAdapter() const final override;
 	virtual IDXGISwapChain*       RHIGetSwapChain(FRHIViewport* InViewport) const final override;
 	virtual DXGI_FORMAT           RHIGetSwapChainFormat(EPixelFormat InFormat) const final override;
 	virtual FTexture2DRHIRef      RHICreateTexture2DFromResource(EPixelFormat Format, ETextureCreateFlags TexCreateFlags, const FClearValueBinding& ClearValueBinding, ID3D11Texture2D* Resource) final override;
diff --git a/Engine/Source/Runtime/Windows/D3D11RHI/Public/ID3D11DynamicRHI.h b/Engine/Source/Runtime/Windows/D3D11RHI/Public/ID3D11DynamicRHI.h
--- a/Engine/Source/Runtime/Windows/D3D11RHI/Public/ID3D11DynamicRHI.h
+++ b/Engine/Source/Runtime/Windows/D3D11RHI/Public/ID3D11DynamicRHI.h
@@ -18,6 +18,7 @@ struct ID3D11DynamicRHI : public FDynamicRHIPSOFallback
 
 	virtual ID3D11Device*         RHIGetDevice() const = 0;
 	virtual ID3D11DeviceContext*  RHIGetDeviceContext() const = 0;
+	virtual IDXGIAdapter*         RHIGetAdapter() const = 0;
 	virtual IDXGISwapChain*       RHIGetSwapChain(FRHIViewport* InViewport) const = 0;
 	virtual DXGI_FORMAT           RHIGetSwapChainFormat(EPixelFormat InFormat) const = 0;
 
