WindowsMixedReality uplugin syntax fix

#rnx

diff --git a/Engine/Plugins/Runtime/WindowsMixedReality/WindowsMixedReality.uplugin b/Engine/Plugins/Runtime/WindowsMixedReality/WindowsMixedReality.uplugin
--- a/Engine/Plugins/Runtime/WindowsMixedReality/WindowsMixedReality.uplugin
+++ b/Engine/Plugins/Runtime/WindowsMixedReality/WindowsMixedReality.uplugin
@@ -55,7 +55,7 @@
 			"Type": "Editor",
 			"LoadingPhase": "PostEngineInit",
 			"PlatformAllowList": [
-				"Win64",
+				"Win64"
 			]
 		},
 		{
