Removed WindowsMixedRealityHMD private pch.
Timing when compiling module with PCH(including compiling PCH): 28.79 secs
Timing using shared PCH: 12.03s secs

diff --git a/Engine/Plugins/Runtime/WindowsMixedReality/Source/WindowsMixedRealityHMD/WindowsMixedRealityHMD.Build.cs b/Engine/Plugins/Runtime/WindowsMixedReality/Source/WindowsMixedRealityHMD/WindowsMixedRealityHMD.Build.cs
--- a/Engine/Plugins/Runtime/WindowsMixedReality/Source/WindowsMixedRealityHMD/WindowsMixedRealityHMD.Build.cs
+++ b/Engine/Plugins/Runtime/WindowsMixedReality/Source/WindowsMixedRealityHMD/WindowsMixedRealityHMD.Build.cs
@@ -103,9 +103,6 @@
 					"../../../../Source/Runtime/Renderer/Private",
 					});
 
-				PCHUsage = PCHUsageMode.NoSharedPCHs;
-				PrivatePCHHeaderFile = "Private/WindowsMixedRealityPrecompiled.h";
-
                 if (Target.Platform == UnrealTargetPlatform.Win64)
                 {
 					RuntimeDependencies.Add(System.IO.Path.Combine("$(EngineDir)/Binaries/ThirdParty/Windows/x64", "CoarseRelocUW.dll"));
