Removed WindowsMixedRealitySpatialInput private pch.
Timing when compiling module with PCH(including compiling PCH): 8.97 secs
Timing using shared PCH: 3.66s secs

diff --git a/Engine/Plugins/Runtime/WindowsMixedReality/Source/WindowsMixedRealitySpatialInput/WindowsMixedRealitySpatialInput.Build.cs b/Engine/Plugins/Runtime/WindowsMixedReality/Source/WindowsMixedRealitySpatialInput/WindowsMixedRealitySpatialInput.Build.cs
--- a/Engine/Plugins/Runtime/WindowsMixedReality/Source/WindowsMixedRealitySpatialInput/WindowsMixedRealitySpatialInput.Build.cs
+++ b/Engine/Plugins/Runtime/WindowsMixedReality/Source/WindowsMixedRealitySpatialInput/WindowsMixedRealitySpatialInput.Build.cs
@@ -63,9 +63,6 @@
             if (Target.Platform == UnrealTargetPlatform.Win64)
             {	
                 bEnableExceptions = true;
-
-                PCHUsage = PCHUsageMode.NoSharedPCHs;
-                PrivatePCHHeaderFile = "Private/WindowsMixedRealitySpatialInput.h";
             }
 
 			if (Target.Platform == UnrealTargetPlatform.Win64 && Target.bBuildEditor == true)
