Update calls to the soon to be deprecated OnController functions

#rnx

diff --git a/Engine/Plugins/Runtime/WindowsMixedReality/Source/WindowsMixedRealitySpatialInput/Private/WindowsMixedRealitySpatialInput.cpp b/Engine/Plugins/Runtime/WindowsMixedReality/Source/WindowsMixedRealitySpatialInput/Private/WindowsMixedRealitySpatialInput.cpp
--- a/Engine/Plugins/Runtime/WindowsMixedReality/Source/WindowsMixedRealitySpatialInput/Private/WindowsMixedRealitySpatialInput.cpp
+++ b/Engine/Plugins/Runtime/WindowsMixedReality/Source/WindowsMixedRealitySpatialInput/Private/WindowsMixedRealitySpatialInput.cpp
@@ -9,6 +9,7 @@
 #include "WindowsMixedRealitySpatialInputFunctionLibrary.h"
 #include "WindowsMixedRealitySpatialInputTypes.h"
 #include "WindowsMixedRealityAvailability.h"
+#include "GenericPlatform/GenericPlatformInputDeviceMapper.h"
 
 #define LOCTEXT_NAMESPACE "WindowsMixedRealitySpatialInput"
 #define MotionControllerDeviceTypeName "WindowsMixedRealitySpatialInput"
@@ -73,6 +74,10 @@ namespace WindowsMixedReality
 		HMDInputPressState pressState) noexcept
 	{
 		check(IsInGameThread());
+		IPlatformInputDeviceMapper& Mapper = IPlatformInputDeviceMapper::Get();
+		FPlatformUserId UserId = FGenericPlatformMisc::GetPlatformUserForUserIndex(controllerId);
+		FInputDeviceId DeviceId = INPUTDEVICEID_NONE;
+		Mapper.RemapControllerIdToPlatformUserAndDevice(controllerId, UserId, DeviceId);
 
 		FName buttonName = button.GetFName();
 
@@ -87,7 +92,8 @@ namespace WindowsMixedReality
 			// Send the press event.
 			messageHandler->OnControllerButtonPressed(
 				buttonName,
-				static_cast<int32>(controllerId),
+				UserId,
+				DeviceId,
 				false);
 		}
 		else
@@ -95,7 +101,8 @@ namespace WindowsMixedReality
 			// Send the release event
 			messageHandler->OnControllerButtonReleased(
 				buttonName,
-				static_cast<int32>(controllerId),
+				UserId,
+				DeviceId,
 				false);
 		}
 	}
@@ -108,11 +115,17 @@ namespace WindowsMixedReality
 	{
 		check(IsInGameThread());
 
+		IPlatformInputDeviceMapper& Mapper = IPlatformInputDeviceMapper::Get();
+		FPlatformUserId UserId = FGenericPlatformMisc::GetPlatformUserForUserIndex(controllerId);
+		FInputDeviceId DeviceId = INPUTDEVICEID_NONE;
+		Mapper.RemapControllerIdToPlatformUserAndDevice(controllerId, UserId, DeviceId);
+
 		FName axisName = axis.GetFName();
 
 		messageHandler->OnControllerAnalog(
 			axisName,
-			static_cast<int32>(controllerId),
+			UserId,
+			DeviceId,
 			static_cast<float>(axisPosition));
 	}
 
diff --git a/Engine/Plugins/Runtime/WindowsMixedReality/Source/WindowsMixedRealitySpatialInput/WindowsMixedRealitySpatialInput.Build.cs b/Engine/Plugins/Runtime/WindowsMixedReality/Source/WindowsMixedRealitySpatialInput/WindowsMixedRealitySpatialInput.Build.cs
--- a/Engine/Plugins/Runtime/WindowsMixedReality/Source/WindowsMixedRealitySpatialInput/WindowsMixedRealitySpatialInput.Build.cs
+++ b/Engine/Plugins/Runtime/WindowsMixedReality/Source/WindowsMixedRealitySpatialInput/WindowsMixedRealitySpatialInput.Build.cs
@@ -50,6 +50,7 @@ namespace UnrealBuildTool.Rules
                     "WindowsMixedRealityHMD",
 					"Slate",
 					"SlateCore",
+					"ApplicationCore"
 				});
 
             AddEngineThirdPartyPrivateStaticDependencies(Target, "WindowsMixedRealityInterop");
