Add dep to Application Core

#rnx

diff --git a/Engine/Plugins/Runtime/WindowsMixedReality/Source/WindowsMixedRealityHandTracking/WindowsMixedRealityHandTracking.Build.cs b/Engine/Plugins/Runtime/WindowsMixedReality/Source/WindowsMixedRealityHandTracking/WindowsMixedRealityHandTracking.Build.cs
--- a/Engine/Plugins/Runtime/WindowsMixedReality/Source/WindowsMixedRealityHandTracking/WindowsMixedRealityHandTracking.Build.cs
+++ b/Engine/Plugins/Runtime/WindowsMixedReality/Source/WindowsMixedRealityHandTracking/WindowsMixedRealityHandTracking.Build.cs
@@ -28,7 +28,8 @@ namespace UnrealBuildTool.Rules
                     "LiveLinkInterface",
 					"Slate",
 					"SlateCore",
-					"WindowsMixedRealityHMD"
+					"WindowsMixedRealityHMD",
+					"ApplicationCore"
 				}
 			);
 
